Adding in the automation test times for those which have failed. The durations were not reflective of this.

TTP 314101 - UFE: Times for automation tests which produce an error are not recorded or included in the total test duration

diff --git a/Engine/Source/Developer/AutomationController/Private/AutomationReport.cpp b/Engine/Source/Developer/AutomationController/Private/AutomationReport.cpp
--- a/Engine/Source/Developer/AutomationController/Private/AutomationReport.cpp
+++ b/Engine/Source/Developer/AutomationController/Private/AutomationReport.cpp
@@ -620,7 +620,8 @@ const bool FAutomationReport::GetDurationRange(float& OutMinTime, float& OutMaxT
 		for( int32 PassIndex = 0; PassIndex < Results[ClusterIndex].Num(); ++PassIndex)
 		{
 			//if we want tests with errors and this test had them OR we want tests warnings and this test had them
-			if( Results[ClusterIndex][PassIndex].State == EAutomationState::Success)
+			if( Results[ClusterIndex][PassIndex].State == EAutomationState::Success ||
+				Results[ClusterIndex][PassIndex].State == EAutomationState::Fail)
 			{
 				OutMinTime = FMath::Min(OutMinTime, Results[ClusterIndex][PassIndex].Duration );
 				OutMaxTime = FMath::Max(OutMaxTime, Results[ClusterIndex][PassIndex].Duration );
