Automation: -Fixed a crash when finding workers in the automation window.

diff --git a/Engine/Source/Developer/AutomationController/Private/AutomationReport.cpp b/Engine/Source/Developer/AutomationController/Private/AutomationReport.cpp
--- a/Engine/Source/Developer/AutomationController/Private/AutomationReport.cpp
+++ b/Engine/Source/Developer/AutomationController/Private/AutomationReport.cpp
@@ -429,11 +429,18 @@ const int32 FAutomationReport::GetNumResults( const int32 ClusterIndex )
 const int32 FAutomationReport::GetCurrentPassIndex( const int32 ClusterIndex )
 {
 	int32 PassIndex = 1;
-	for(; PassIndex < Results[ClusterIndex].Num(); ++PassIndex )
+
+	if( IsSupported(ClusterIndex) )
 	{
-		if( Results[ClusterIndex][PassIndex].State == EAutomationState::NotRun )
-			break;
+		for(; PassIndex < Results[ClusterIndex].Num(); ++PassIndex )
+		{
+			if( Results[ClusterIndex][PassIndex].State == EAutomationState::NotRun )
+			{
+				break;
+			}
+		}
 	}
+
 	return PassIndex - 1;
 }
 
