UE4 Refactoring. Changed OVERRIDE and FINAL macros to keywords override and final.

diff --git a/Engine/Source/Developer/AssetTools/Private/AssetTypeActions/AssetTypeActions_CurveVector.h b/Engine/Source/Developer/AssetTools/Private/AssetTypeActions/AssetTypeActions_CurveVector.h
--- a/Engine/Source/Developer/AssetTools/Private/AssetTypeActions/AssetTypeActions_CurveVector.h
+++ b/Engine/Source/Developer/AssetTools/Private/AssetTypeActions/AssetTypeActions_CurveVector.h
@@ -6,7 +6,7 @@ class FAssetTypeActions_CurveVector : public FAssetTypeActions_Curve
 {
 public:
 	// IAssetTypeActions Implementation
-	virtual FText GetName() const OVERRIDE { return NSLOCTEXT("AssetTypeActions", "AssetTypeActions_CurveVector", "Vector Curve"); }
-	virtual UClass* GetSupportedClass() const OVERRIDE { return UCurveVector::StaticClass(); }
-	virtual bool CanFilter() OVERRIDE { return true; }
+	virtual FText GetName() const override { return NSLOCTEXT("AssetTypeActions", "AssetTypeActions_CurveVector", "Vector Curve"); }
+	virtual UClass* GetSupportedClass() const override { return UCurveVector::StaticClass(); }
+	virtual bool CanFilter() override { return true; }
 };
\ No newline at end of file
diff --git a/Engine/Source/Developer/AssetTools/Private/AssetTypeActions/AssetTypeActions_DataAsset.h b/Engine/Source/Developer/AssetTools/Private/AssetTypeActions/AssetTypeActions_DataAsset.h
--- a/Engine/Source/Developer/AssetTools/Private/AssetTypeActions/AssetTypeActions_DataAsset.h
+++ b/Engine/Source/Developer/AssetTools/Private/AssetTypeActions/AssetTypeActions_DataAsset.h
@@ -6,8 +6,8 @@ class FAssetTypeActions_DataAsset : public FAssetTypeActions_Base
 {
 public:
 	// IAssetTypeActions Implementation
-	virtual FText GetName() const OVERRIDE { return NSLOCTEXT("AssetTypeActions", "AssetTypeActions_DataAsset", "Data Asset"); }
-	virtual FColor GetTypeColor() const OVERRIDE { return FColor(201, 29, 85); }
-	virtual UClass* GetSupportedClass() const OVERRIDE { return UDataAsset::StaticClass(); }
-	virtual uint32 GetCategories() OVERRIDE { return EAssetTypeCategories::Misc; }
+	virtual FText GetName() const override { return NSLOCTEXT("AssetTypeActions", "AssetTypeActions_DataAsset", "Data Asset"); }
+	virtual FColor GetTypeColor() const override { return FColor(201, 29, 85); }
+	virtual UClass* GetSupportedClass() const override { return UDataAsset::StaticClass(); }
+	virtual uint32 GetCategories() override { return EAssetTypeCategories::Misc; }
 };
\ No newline at end of file
diff --git a/Engine/Source/Developer/AssetTools/Private/AssetTypeActions/AssetTypeActions_DataTable.h b/Engine/Source/Developer/AssetTools/Private/AssetTypeActions/AssetTypeActions_DataTable.h
--- a/Engine/Source/Developer/AssetTools/Private/AssetTypeActions/AssetTypeActions_DataTable.h
+++ b/Engine/Source/Developer/AssetTools/Private/AssetTypeActions/AssetTypeActions_DataTable.h
@@ -6,13 +6,13 @@ class FAssetTypeActions_DataTable : public FAssetTypeActions_CSVAssetBase
 {
 public:
 	// IAssetTypeActions Implementation
-	virtual FText GetName() const OVERRIDE { return NSLOCTEXT("AssetTypeActions", "AssetTypeActions_DataTable", "Data Table"); }
-	virtual FColor GetTypeColor() const OVERRIDE { return FColor(62, 140, 35); }
-	virtual UClass* GetSupportedClass() const OVERRIDE { return UDataTable::StaticClass(); }
-	virtual bool HasActions ( const TArray<UObject*>& InObjects ) const OVERRIDE { return true; }
-	virtual void GetActions( const TArray<UObject*>& InObjects, FMenuBuilder& MenuBuilder ) OVERRIDE;
-	virtual void OpenAssetEditor( const TArray<UObject*>& InObjects, TSharedPtr<class IToolkitHost> EditWithinLevelEditor = TSharedPtr<IToolkitHost>() ) OVERRIDE;
-	virtual uint32 GetCategories() OVERRIDE { return EAssetTypeCategories::Misc; }
+	virtual FText GetName() const override { return NSLOCTEXT("AssetTypeActions", "AssetTypeActions_DataTable", "Data Table"); }
+	virtual FColor GetTypeColor() const override { return FColor(62, 140, 35); }
+	virtual UClass* GetSupportedClass() const override { return UDataTable::StaticClass(); }
+	virtual bool HasActions ( const TArray<UObject*>& InObjects ) const override { return true; }
+	virtual void GetActions( const TArray<UObject*>& InObjects, FMenuBuilder& MenuBuilder ) override;
+	virtual void OpenAssetEditor( const TArray<UObject*>& InObjects, TSharedPtr<class IToolkitHost> EditWithinLevelEditor = TSharedPtr<IToolkitHost>() ) override;
+	virtual uint32 GetCategories() override { return EAssetTypeCategories::Misc; }
 
 	/** Handler for when JSON is selected */
 	void ExecuteJSON(TArray< TWeakObjectPtr<UObject> > Objects);
diff --git a/Engine/Source/Developer/AssetTools/Private/AssetTypeActions/AssetTypeActions_DestructibleMesh.h b/Engine/Source/Developer/AssetTools/Private/AssetTypeActions/AssetTypeActions_DestructibleMesh.h
--- a/Engine/Source/Developer/AssetTools/Private/AssetTypeActions/AssetTypeActions_DestructibleMesh.h
+++ b/Engine/Source/Developer/AssetTools/Private/AssetTypeActions/AssetTypeActions_DestructibleMesh.h
@@ -6,12 +6,12 @@ class FAssetTypeActions_DestructibleMesh : public FAssetTypeActions_SkeletalMesh
 {
 public:
 	// IAssetTypeActions Implementation
-	virtual FText GetName() const OVERRIDE { return NSLOCTEXT("AssetTypeActions", "AssetTypeActions_DestructibleMesh", "Destructible Mesh"); }
-	virtual FColor GetTypeColor() const OVERRIDE { return FColor(200,128,128); }
-	virtual UClass* GetSupportedClass() const OVERRIDE;
-	virtual uint32 GetCategories() OVERRIDE { return EAssetTypeCategories::Physics; }
-	virtual void OpenAssetEditor( const TArray<UObject*>& InObjects, TSharedPtr<class IToolkitHost> EditWithinLevelEditor = TSharedPtr<IToolkitHost>() ) OVERRIDE;
-	virtual void GetActions( const TArray<UObject*>& InObjects, FMenuBuilder& MenuBuilder ) OVERRIDE;
+	virtual FText GetName() const override { return NSLOCTEXT("AssetTypeActions", "AssetTypeActions_DestructibleMesh", "Destructible Mesh"); }
+	virtual FColor GetTypeColor() const override { return FColor(200,128,128); }
+	virtual UClass* GetSupportedClass() const override;
+	virtual uint32 GetCategories() override { return EAssetTypeCategories::Physics; }
+	virtual void OpenAssetEditor( const TArray<UObject*>& InObjects, TSharedPtr<class IToolkitHost> EditWithinLevelEditor = TSharedPtr<IToolkitHost>() ) override;
+	virtual void GetActions( const TArray<UObject*>& InObjects, FMenuBuilder& MenuBuilder ) override;
 
 	/** Handler for when Edit is selected */
 	void ExecuteEdit(TArray<TWeakObjectPtr<class UDestructibleMesh>> Objects);
diff --git a/Engine/Source/Developer/AssetTools/Private/AssetTypeActions/AssetTypeActions_DialogueVoice.h b/Engine/Source/Developer/AssetTools/Private/AssetTypeActions/AssetTypeActions_DialogueVoice.h
--- a/Engine/Source/Developer/AssetTools/Private/AssetTypeActions/AssetTypeActions_DialogueVoice.h
+++ b/Engine/Source/Developer/AssetTools/Private/AssetTypeActions/AssetTypeActions_DialogueVoice.h
@@ -8,9 +8,9 @@ class FAssetTypeActions_DialogueVoice : public FAssetTypeActions_Base
 {
 public:
 	// IAssetTypeActions Implementation
-	virtual FText GetName() const OVERRIDE { return NSLOCTEXT("AssetTypeActions", "AssetTypeActions_DialogueVoice", "Dialogue Voice"); }
-	virtual FColor GetTypeColor() const OVERRIDE { return FColor(0,0,255); }
-	virtual UClass* GetSupportedClass() const OVERRIDE { return UDialogueVoice::StaticClass(); }
-	virtual uint32 GetCategories() OVERRIDE { return EAssetTypeCategories::Sounds; }
-	virtual bool CanFilter() OVERRIDE { return true; }
+	virtual FText GetName() const override { return NSLOCTEXT("AssetTypeActions", "AssetTypeActions_DialogueVoice", "Dialogue Voice"); }
+	virtual FColor GetTypeColor() const override { return FColor(0,0,255); }
+	virtual UClass* GetSupportedClass() const override { return UDialogueVoice::StaticClass(); }
+	virtual uint32 GetCategories() override { return EAssetTypeCategories::Sounds; }
+	virtual bool CanFilter() override { return true; }
 };
\ No newline at end of file
diff --git a/Engine/Source/Developer/AssetTools/Private/AssetTypeActions/AssetTypeActions_DialogueWave.h b/Engine/Source/Developer/AssetTools/Private/AssetTypeActions/AssetTypeActions_DialogueWave.h
--- a/Engine/Source/Developer/AssetTools/Private/AssetTypeActions/AssetTypeActions_DialogueWave.h
+++ b/Engine/Source/Developer/AssetTools/Private/AssetTypeActions/AssetTypeActions_DialogueWave.h
@@ -8,14 +8,14 @@ class FAssetTypeActions_DialogueWave : public FAssetTypeActions_Base
 {
 public:
 	// IAssetTypeActions Implementation
-	virtual FText GetName() const OVERRIDE { return NSLOCTEXT("AssetTypeActions", "AssetTypeActions_DialogueWave", "Dialogue Wave"); }
-	virtual FColor GetTypeColor() const OVERRIDE { return FColor(0,0,255); }
-	virtual UClass* GetSupportedClass() const OVERRIDE { return UDialogueWave::StaticClass(); }
-	virtual bool HasActions ( const TArray<UObject*>& InObjects ) const OVERRIDE { return true; }
-	virtual void GetActions( const TArray<UObject*>& InObjects, FMenuBuilder& MenuBuilder ) OVERRIDE;
-	virtual uint32 GetCategories() OVERRIDE { return EAssetTypeCategories::Sounds; }
-	virtual bool CanFilter() OVERRIDE { return true; }
-	virtual void OpenAssetEditor( const TArray<UObject*>& InObjects, TSharedPtr<class IToolkitHost> EditWithinLevelEditor = TSharedPtr<IToolkitHost>() ) OVERRIDE;
+	virtual FText GetName() const override { return NSLOCTEXT("AssetTypeActions", "AssetTypeActions_DialogueWave", "Dialogue Wave"); }
+	virtual FColor GetTypeColor() const override { return FColor(0,0,255); }
+	virtual UClass* GetSupportedClass() const override { return UDialogueWave::StaticClass(); }
+	virtual bool HasActions ( const TArray<UObject*>& InObjects ) const override { return true; }
+	virtual void GetActions( const TArray<UObject*>& InObjects, FMenuBuilder& MenuBuilder ) override;
+	virtual uint32 GetCategories() override { return EAssetTypeCategories::Sounds; }
+	virtual bool CanFilter() override { return true; }
+	virtual void OpenAssetEditor( const TArray<UObject*>& InObjects, TSharedPtr<class IToolkitHost> EditWithinLevelEditor = TSharedPtr<IToolkitHost>() ) override;
 
 private:
 	/** Handler for when Edit is selected */
diff --git a/Engine/Source/Developer/AssetTools/Private/AssetTypeActions/AssetTypeActions_Enum.h b/Engine/Source/Developer/AssetTools/Private/AssetTypeActions/AssetTypeActions_Enum.h
--- a/Engine/Source/Developer/AssetTools/Private/AssetTypeActions/AssetTypeActions_Enum.h
+++ b/Engine/Source/Developer/AssetTools/Private/AssetTypeActions/AssetTypeActions_Enum.h
@@ -6,10 +6,10 @@ class FAssetTypeActions_Enum : public FAssetTypeActions_Base
 {
 public:
 	// IAssetTypeActions Implementation
-	virtual FText GetName() const OVERRIDE { return NSLOCTEXT("AssetTypeActions", "AssetTypeActions_Enum", "Enumeration"); }
-	virtual FColor GetTypeColor() const OVERRIDE { return FColor(255, 200, 200); }
-	virtual UClass* GetSupportedClass() const OVERRIDE { return UUserDefinedEnum::StaticClass(); }
-	virtual uint32 GetCategories() OVERRIDE { return EAssetTypeCategories::Misc; }
+	virtual FText GetName() const override { return NSLOCTEXT("AssetTypeActions", "AssetTypeActions_Enum", "Enumeration"); }
+	virtual FColor GetTypeColor() const override { return FColor(255, 200, 200); }
+	virtual UClass* GetSupportedClass() const override { return UUserDefinedEnum::StaticClass(); }
+	virtual uint32 GetCategories() override { return EAssetTypeCategories::Misc; }
 
-	virtual void OpenAssetEditor( const TArray<UObject*>& InObjects, TSharedPtr<class IToolkitHost> EditWithinLevelEditor = TSharedPtr<IToolkitHost>() ) OVERRIDE;
+	virtual void OpenAssetEditor( const TArray<UObject*>& InObjects, TSharedPtr<class IToolkitHost> EditWithinLevelEditor = TSharedPtr<IToolkitHost>() ) override;
 };
diff --git a/Engine/Source/Developer/AssetTools/Private/AssetTypeActions/AssetTypeActions_Font.h b/Engine/Source/Developer/AssetTools/Private/AssetTypeActions/AssetTypeActions_Font.h
--- a/Engine/Source/Developer/AssetTools/Private/AssetTypeActions/AssetTypeActions_Font.h
+++ b/Engine/Source/Developer/AssetTools/Private/AssetTypeActions/AssetTypeActions_Font.h
@@ -6,13 +6,13 @@ class FAssetTypeActions_Font : public FAssetTypeActions_Base
 {
 public:
 	// IAssetTypeActions Implementation
-	virtual FText GetName() const OVERRIDE { return NSLOCTEXT("AssetTypeActions", "AssetTypeActions_Font", "Font"); }
-	virtual FColor GetTypeColor() const OVERRIDE { return FColor(255,0,0); }
-	virtual UClass* GetSupportedClass() const OVERRIDE { return UFont::StaticClass(); }
-	virtual bool HasActions ( const TArray<UObject*>& InObjects ) const OVERRIDE { return true; }
-	virtual void GetActions( const TArray<UObject*>& InObjects, FMenuBuilder& MenuBuilder ) OVERRIDE;
-	virtual void OpenAssetEditor( const TArray<UObject*>& InObjects, TSharedPtr<class IToolkitHost> EditWithinLevelEditor = TSharedPtr<IToolkitHost>() ) OVERRIDE;
-	virtual uint32 GetCategories() OVERRIDE { return EAssetTypeCategories::MaterialsAndTextures; }
+	virtual FText GetName() const override { return NSLOCTEXT("AssetTypeActions", "AssetTypeActions_Font", "Font"); }
+	virtual FColor GetTypeColor() const override { return FColor(255,0,0); }
+	virtual UClass* GetSupportedClass() const override { return UFont::StaticClass(); }
+	virtual bool HasActions ( const TArray<UObject*>& InObjects ) const override { return true; }
+	virtual void GetActions( const TArray<UObject*>& InObjects, FMenuBuilder& MenuBuilder ) override;
+	virtual void OpenAssetEditor( const TArray<UObject*>& InObjects, TSharedPtr<class IToolkitHost> EditWithinLevelEditor = TSharedPtr<IToolkitHost>() ) override;
+	virtual uint32 GetCategories() override { return EAssetTypeCategories::MaterialsAndTextures; }
 
 private:
 	/** Handler for when Edit is selected */
diff --git a/Engine/Source/Developer/AssetTools/Private/AssetTypeActions/AssetTypeActions_ForceFeedbackEffect.h b/Engine/Source/Developer/AssetTools/Private/AssetTypeActions/AssetTypeActions_ForceFeedbackEffect.h
--- a/Engine/Source/Developer/AssetTools/Private/AssetTypeActions/AssetTypeActions_ForceFeedbackEffect.h
+++ b/Engine/Source/Developer/AssetTools/Private/AssetTypeActions/AssetTypeActions_ForceFeedbackEffect.h
@@ -6,8 +6,8 @@ class FAssetTypeActions_ForceFeedbackEffect : public FAssetTypeActions_Base
 {
 public:
 	// IAssetTypeActions Implementation
-	virtual FText GetName() const OVERRIDE { return NSLOCTEXT("AssetTypeActions", "AssetTypeActions_ForceFeedbackEffect", "Force Feedback Effect"); }
-	virtual FColor GetTypeColor() const OVERRIDE { return FColor(175, 0, 0); }
-	virtual UClass* GetSupportedClass() const OVERRIDE;
-	virtual uint32 GetCategories() OVERRIDE { return EAssetTypeCategories::Misc; }
+	virtual FText GetName() const override { return NSLOCTEXT("AssetTypeActions", "AssetTypeActions_ForceFeedbackEffect", "Force Feedback Effect"); }
+	virtual FColor GetTypeColor() const override { return FColor(175, 0, 0); }
+	virtual UClass* GetSupportedClass() const override;
+	virtual uint32 GetCategories() override { return EAssetTypeCategories::Misc; }
 };
\ No newline at end of file
diff --git a/Engine/Source/Developer/AssetTools/Private/AssetTypeActions/AssetTypeActions_InstancedFoliageSettings.h b/Engine/Source/Developer/AssetTools/Private/AssetTypeActions/AssetTypeActions_InstancedFoliageSettings.h
--- a/Engine/Source/Developer/AssetTools/Private/AssetTypeActions/AssetTypeActions_InstancedFoliageSettings.h
+++ b/Engine/Source/Developer/AssetTools/Private/AssetTypeActions/AssetTypeActions_InstancedFoliageSettings.h
@@ -6,8 +6,8 @@ class FAssetTypeActions_InstancedFoliageSettings : public FAssetTypeActions_Base
 {
 public:
 	// IAssetTypeActions Implementation
-	virtual FText GetName() const OVERRIDE { return NSLOCTEXT("AssetTypeActions", "AssetTypeActions_InstancedFoliageSettings", "Foliage Settings"); }
-	virtual FColor GetTypeColor() const OVERRIDE { return FColor(128,192,255); }
-	virtual UClass* GetSupportedClass() const OVERRIDE;
-	virtual uint32 GetCategories() OVERRIDE { return EAssetTypeCategories::Misc; }
+	virtual FText GetName() const override { return NSLOCTEXT("AssetTypeActions", "AssetTypeActions_InstancedFoliageSettings", "Foliage Settings"); }
+	virtual FColor GetTypeColor() const override { return FColor(128,192,255); }
+	virtual UClass* GetSupportedClass() const override;
+	virtual uint32 GetCategories() override { return EAssetTypeCategories::Misc; }
 };
\ No newline at end of file
diff --git a/Engine/Source/Developer/AssetTools/Private/AssetTypeActions/AssetTypeActions_InterpData.h b/Engine/Source/Developer/AssetTools/Private/AssetTypeActions/AssetTypeActions_InterpData.h
--- a/Engine/Source/Developer/AssetTools/Private/AssetTypeActions/AssetTypeActions_InterpData.h
+++ b/Engine/Source/Developer/AssetTools/Private/AssetTypeActions/AssetTypeActions_InterpData.h
@@ -8,8 +8,8 @@ class FAssetTypeActions_InterpData : public FAssetTypeActions_Base
 {
 public:
 	// IAssetTypeActions Implementation
-	virtual FText GetName() const OVERRIDE { return NSLOCTEXT("AssetTypeActions", "AssetTypeActions_InterpData", "Matinee Data"); }
-	virtual FColor GetTypeColor() const OVERRIDE { return FColor(255,128,0); }
-	virtual UClass* GetSupportedClass() const OVERRIDE { return UInterpData::StaticClass(); }
-	virtual uint32 GetCategories() OVERRIDE { return EAssetTypeCategories::Misc; }
+	virtual FText GetName() const override { return NSLOCTEXT("AssetTypeActions", "AssetTypeActions_InterpData", "Matinee Data"); }
+	virtual FColor GetTypeColor() const override { return FColor(255,128,0); }
+	virtual UClass* GetSupportedClass() const override { return UInterpData::StaticClass(); }
+	virtual uint32 GetCategories() override { return EAssetTypeCategories::Misc; }
 };
\ No newline at end of file
diff --git a/Engine/Source/Developer/AssetTools/Private/AssetTypeActions/AssetTypeActions_LandscapeLayer.h b/Engine/Source/Developer/AssetTools/Private/AssetTypeActions/AssetTypeActions_LandscapeLayer.h
--- a/Engine/Source/Developer/AssetTools/Private/AssetTypeActions/AssetTypeActions_LandscapeLayer.h
+++ b/Engine/Source/Developer/AssetTools/Private/AssetTypeActions/AssetTypeActions_LandscapeLayer.h
@@ -8,8 +8,8 @@ class FAssetTypeActions_LandscapeLayer : public FAssetTypeActions_Base
 {
 public:
 	// IAssetTypeActions Implementation
-	virtual FText GetName() const OVERRIDE { return NSLOCTEXT("AssetTypeActions", "AssetTypeActions_LandscapeLayer", "Landscape Layer"); }
-	virtual FColor GetTypeColor() const OVERRIDE { return FColor(128,192,255); }
-	virtual UClass* GetSupportedClass() const OVERRIDE { return ULandscapeLayerInfoObject::StaticClass(); }
-	virtual uint32 GetCategories() OVERRIDE { return EAssetTypeCategories::Misc; }
+	virtual FText GetName() const override { return NSLOCTEXT("AssetTypeActions", "AssetTypeActions_LandscapeLayer", "Landscape Layer"); }
+	virtual FColor GetTypeColor() const override { return FColor(128,192,255); }
+	virtual UClass* GetSupportedClass() const override { return ULandscapeLayerInfoObject::StaticClass(); }
+	virtual uint32 GetCategories() override { return EAssetTypeCategories::Misc; }
 };
\ No newline at end of file
diff --git a/Engine/Source/Developer/AssetTools/Private/AssetTypeActions/AssetTypeActions_Material.h b/Engine/Source/Developer/AssetTools/Private/AssetTypeActions/AssetTypeActions_Material.h
--- a/Engine/Source/Developer/AssetTools/Private/AssetTypeActions/AssetTypeActions_Material.h
+++ b/Engine/Source/Developer/AssetTools/Private/AssetTypeActions/AssetTypeActions_Material.h
@@ -6,13 +6,13 @@ class FAssetTypeActions_Material : public FAssetTypeActions_MaterialInterface
 {
 public:
 	// IAssetTypeActions Implementation
-	virtual FText GetName() const OVERRIDE { return NSLOCTEXT("AssetTypeActions", "AssetTypeActions_Material", "Material"); }
-	virtual FColor GetTypeColor() const OVERRIDE { return FColor(0,255,0); }
-	virtual UClass* GetSupportedClass() const OVERRIDE { return UMaterial::StaticClass(); }
-	virtual void GetActions( const TArray<UObject*>& InObjects, FMenuBuilder& MenuBuilder ) OVERRIDE;
-	virtual void OpenAssetEditor( const TArray<UObject*>& InObjects, TSharedPtr<class IToolkitHost> EditWithinLevelEditor = TSharedPtr<IToolkitHost>() ) OVERRIDE;
-	virtual bool CanFilter() OVERRIDE { return true; }
-	virtual uint32 GetCategories() OVERRIDE { return FAssetTypeActions_MaterialInterface::GetCategories() | EAssetTypeCategories::Basic; }
+	virtual FText GetName() const override { return NSLOCTEXT("AssetTypeActions", "AssetTypeActions_Material", "Material"); }
+	virtual FColor GetTypeColor() const override { return FColor(0,255,0); }
+	virtual UClass* GetSupportedClass() const override { return UMaterial::StaticClass(); }
+	virtual void GetActions( const TArray<UObject*>& InObjects, FMenuBuilder& MenuBuilder ) override;
+	virtual void OpenAssetEditor( const TArray<UObject*>& InObjects, TSharedPtr<class IToolkitHost> EditWithinLevelEditor = TSharedPtr<IToolkitHost>() ) override;
+	virtual bool CanFilter() override { return true; }
+	virtual uint32 GetCategories() override { return FAssetTypeActions_MaterialInterface::GetCategories() | EAssetTypeCategories::Basic; }
 
 private:
 	/** Handler for when Edit is selected */
diff --git a/Engine/Source/Developer/AssetTools/Private/AssetTypeActions/AssetTypeActions_MaterialFunction.h b/Engine/Source/Developer/AssetTools/Private/AssetTypeActions/AssetTypeActions_MaterialFunction.h
--- a/Engine/Source/Developer/AssetTools/Private/AssetTypeActions/AssetTypeActions_MaterialFunction.h
+++ b/Engine/Source/Developer/AssetTools/Private/AssetTypeActions/AssetTypeActions_MaterialFunction.h
@@ -8,14 +8,14 @@ class FAssetTypeActions_MaterialFunction : public FAssetTypeActions_Base
 {
 public:
 	// IAssetTypeActions Implementation
-	virtual FText GetName() const OVERRIDE { return NSLOCTEXT("AssetTypeActions", "AssetTypeActions_MaterialFunction", "Material Function"); }
-	virtual FColor GetTypeColor() const OVERRIDE { return FColor(0,175,175); }
-	virtual UClass* GetSupportedClass() const OVERRIDE { return UMaterialFunction::StaticClass(); }
-	virtual bool HasActions ( const TArray<UObject*>& InObjects ) const OVERRIDE { return true; }
-	virtual void GetActions( const TArray<UObject*>& InObjects, FMenuBuilder& MenuBuilder ) OVERRIDE;
-	virtual void OpenAssetEditor( const TArray<UObject*>& InObjects, TSharedPtr<class IToolkitHost> EditWithinLevelEditor = TSharedPtr<IToolkitHost>() ) OVERRIDE;
-	virtual uint32 GetCategories() OVERRIDE { return EAssetTypeCategories::MaterialsAndTextures; }
-	virtual class UThumbnailInfo* GetThumbnailInfo(UObject* Asset) const OVERRIDE;
+	virtual FText GetName() const override { return NSLOCTEXT("AssetTypeActions", "AssetTypeActions_MaterialFunction", "Material Function"); }
+	virtual FColor GetTypeColor() const override { return FColor(0,175,175); }
+	virtual UClass* GetSupportedClass() const override { return UMaterialFunction::StaticClass(); }
+	virtual bool HasActions ( const TArray<UObject*>& InObjects ) const override { return true; }
+	virtual void GetActions( const TArray<UObject*>& InObjects, FMenuBuilder& MenuBuilder ) override;
+	virtual void OpenAssetEditor( const TArray<UObject*>& InObjects, TSharedPtr<class IToolkitHost> EditWithinLevelEditor = TSharedPtr<IToolkitHost>() ) override;
+	virtual uint32 GetCategories() override { return EAssetTypeCategories::MaterialsAndTextures; }
+	virtual class UThumbnailInfo* GetThumbnailInfo(UObject* Asset) const override;
 
 private:
 	/** Handler for when Edit is selected */
diff --git a/Engine/Source/Developer/AssetTools/Private/AssetTypeActions/AssetTypeActions_MaterialInstanceConstant.h b/Engine/Source/Developer/AssetTools/Private/AssetTypeActions/AssetTypeActions_MaterialInstanceConstant.h
--- a/Engine/Source/Developer/AssetTools/Private/AssetTypeActions/AssetTypeActions_MaterialInstanceConstant.h
+++ b/Engine/Source/Developer/AssetTools/Private/AssetTypeActions/AssetTypeActions_MaterialInstanceConstant.h
@@ -6,12 +6,12 @@ class FAssetTypeActions_MaterialInstanceConstant : public FAssetTypeActions_Mate
 {
 public:
 	// IAssetTypeActions Implementation
-	virtual FText GetName() const OVERRIDE { return NSLOCTEXT("AssetTypeActions", "AssetTypeActions_MaterialInstanceConstant", "Material Instance"); }
-	virtual FColor GetTypeColor() const OVERRIDE { return FColor(0,128,0); }
-	virtual UClass* GetSupportedClass() const OVERRIDE { return UMaterialInstanceConstant::StaticClass(); }
-	virtual void GetActions( const TArray<UObject*>& InObjects, FMenuBuilder& MenuBuilder ) OVERRIDE;
-	virtual void OpenAssetEditor( const TArray<UObject*>& InObjects, TSharedPtr<class IToolkitHost> EditWithinLevelEditor = TSharedPtr<IToolkitHost>() ) OVERRIDE;
-	virtual bool CanFilter() OVERRIDE { return true; }
+	virtual FText GetName() const override { return NSLOCTEXT("AssetTypeActions", "AssetTypeActions_MaterialInstanceConstant", "Material Instance"); }
+	virtual FColor GetTypeColor() const override { return FColor(0,128,0); }
+	virtual UClass* GetSupportedClass() const override { return UMaterialInstanceConstant::StaticClass(); }
+	virtual void GetActions( const TArray<UObject*>& InObjects, FMenuBuilder& MenuBuilder ) override;
+	virtual void OpenAssetEditor( const TArray<UObject*>& InObjects, TSharedPtr<class IToolkitHost> EditWithinLevelEditor = TSharedPtr<IToolkitHost>() ) override;
+	virtual bool CanFilter() override { return true; }
 
 private:
 	/** Handler for when Edit is selected */
diff --git a/Engine/Source/Developer/AssetTools/Private/AssetTypeActions/AssetTypeActions_MaterialInterface.h b/Engine/Source/Developer/AssetTools/Private/AssetTypeActions/AssetTypeActions_MaterialInterface.h
--- a/Engine/Source/Developer/AssetTools/Private/AssetTypeActions/AssetTypeActions_MaterialInterface.h
+++ b/Engine/Source/Developer/AssetTools/Private/AssetTypeActions/AssetTypeActions_MaterialInterface.h
@@ -6,14 +6,14 @@ class FAssetTypeActions_MaterialInterface : public FAssetTypeActions_Base
 {
 public:
 	// IAssetTypeActions Implementation
-	virtual FText GetName() const OVERRIDE { return NSLOCTEXT("AssetTypeActions", "AssetTypeActions_MaterialInterface", "Material Interface"); }
-	virtual FColor GetTypeColor() const OVERRIDE { return FColor(0,255,0); }
-	virtual UClass* GetSupportedClass() const OVERRIDE { return UMaterialInterface::StaticClass(); }
-	virtual bool HasActions ( const TArray<UObject*>& InObjects ) const OVERRIDE { return true; }
-	virtual void GetActions( const TArray<UObject*>& InObjects, FMenuBuilder& MenuBuilder ) OVERRIDE;
-	virtual bool CanFilter() OVERRIDE { return false; }
-	virtual uint32 GetCategories() OVERRIDE { return EAssetTypeCategories::MaterialsAndTextures; }
-	virtual class UThumbnailInfo* GetThumbnailInfo(UObject* Asset) const OVERRIDE;
+	virtual FText GetName() const override { return NSLOCTEXT("AssetTypeActions", "AssetTypeActions_MaterialInterface", "Material Interface"); }
+	virtual FColor GetTypeColor() const override { return FColor(0,255,0); }
+	virtual UClass* GetSupportedClass() const override { return UMaterialInterface::StaticClass(); }
+	virtual bool HasActions ( const TArray<UObject*>& InObjects ) const override { return true; }
+	virtual void GetActions( const TArray<UObject*>& InObjects, FMenuBuilder& MenuBuilder ) override;
+	virtual bool CanFilter() override { return false; }
+	virtual uint32 GetCategories() override { return EAssetTypeCategories::MaterialsAndTextures; }
+	virtual class UThumbnailInfo* GetThumbnailInfo(UObject* Asset) const override;
 
 private:
 	/** Handler for when NewMIC is selected */
diff --git a/Engine/Source/Developer/AssetTools/Private/AssetTypeActions/AssetTypeActions_MaterialParameterCollection.h b/Engine/Source/Developer/AssetTools/Private/AssetTypeActions/AssetTypeActions_MaterialParameterCollection.h
--- a/Engine/Source/Developer/AssetTools/Private/AssetTypeActions/AssetTypeActions_MaterialParameterCollection.h
+++ b/Engine/Source/Developer/AssetTools/Private/AssetTypeActions/AssetTypeActions_MaterialParameterCollection.h
@@ -6,8 +6,8 @@ class FAssetTypeActions_MaterialParameterCollection : public FAssetTypeActions_B
 {
 public:
 	// IAssetTypeActions Implementation
-	virtual FText GetName() const OVERRIDE { return NSLOCTEXT("AssetTypeActions", "AssetTypeActions_MaterialParameterCollection", "Material Parameter Collection"); }
-	virtual FColor GetTypeColor() const OVERRIDE { return FColor(0, 192, 0); }
-	virtual UClass* GetSupportedClass() const OVERRIDE { return UMaterialParameterCollection::StaticClass(); }
-	virtual uint32 GetCategories() OVERRIDE { return EAssetTypeCategories::MaterialsAndTextures; }
+	virtual FText GetName() const override { return NSLOCTEXT("AssetTypeActions", "AssetTypeActions_MaterialParameterCollection", "Material Parameter Collection"); }
+	virtual FColor GetTypeColor() const override { return FColor(0, 192, 0); }
+	virtual UClass* GetSupportedClass() const override { return UMaterialParameterCollection::StaticClass(); }
+	virtual uint32 GetCategories() override { return EAssetTypeCategories::MaterialsAndTextures; }
 };
\ No newline at end of file
diff --git a/Engine/Source/Developer/AssetTools/Private/AssetTypeActions/AssetTypeActions_MorphTarget.h b/Engine/Source/Developer/AssetTools/Private/AssetTypeActions/AssetTypeActions_MorphTarget.h
--- a/Engine/Source/Developer/AssetTools/Private/AssetTypeActions/AssetTypeActions_MorphTarget.h
+++ b/Engine/Source/Developer/AssetTools/Private/AssetTypeActions/AssetTypeActions_MorphTarget.h
@@ -6,13 +6,13 @@ class FAssetTypeActions_MorphTarget : public FAssetTypeActions_Base
 {
 public:
 	// IAssetTypeActions Implementation
-	virtual FText GetName() const OVERRIDE { return NSLOCTEXT("AssetTypeActions", "AssetTypeActions_MorphTarget", "Morph Target"); }
-	virtual FColor GetTypeColor() const OVERRIDE { return FColor(87,0,174); }
-	virtual UClass* GetSupportedClass() const OVERRIDE { return UMorphTarget::StaticClass(); }
-	virtual bool HasActions ( const TArray<UObject*>& InObjects ) const OVERRIDE { return true; }
-	virtual void GetActions( const TArray<UObject*>& InObjects, FMenuBuilder& MenuBuilder ) OVERRIDE;
-	virtual void OpenAssetEditor( const TArray<UObject*>& InObjects, TSharedPtr<class IToolkitHost> EditWithinLevelEditor = TSharedPtr<IToolkitHost>() ) OVERRIDE;
-	virtual uint32 GetCategories() OVERRIDE { return EAssetTypeCategories::Animation; }
+	virtual FText GetName() const override { return NSLOCTEXT("AssetTypeActions", "AssetTypeActions_MorphTarget", "Morph Target"); }
+	virtual FColor GetTypeColor() const override { return FColor(87,0,174); }
+	virtual UClass* GetSupportedClass() const override { return UMorphTarget::StaticClass(); }
+	virtual bool HasActions ( const TArray<UObject*>& InObjects ) const override { return true; }
+	virtual void GetActions( const TArray<UObject*>& InObjects, FMenuBuilder& MenuBuilder ) override;
+	virtual void OpenAssetEditor( const TArray<UObject*>& InObjects, TSharedPtr<class IToolkitHost> EditWithinLevelEditor = TSharedPtr<IToolkitHost>() ) override;
+	virtual uint32 GetCategories() override { return EAssetTypeCategories::Animation; }
 
 private:
 	/** Handler for when Edit is selected */
diff --git a/Engine/Source/Developer/AssetTools/Private/AssetTypeActions/AssetTypeActions_NiagaraScript.h b/Engine/Source/Developer/AssetTools/Private/AssetTypeActions/AssetTypeActions_NiagaraScript.h
--- a/Engine/Source/Developer/AssetTools/Private/AssetTypeActions/AssetTypeActions_NiagaraScript.h
+++ b/Engine/Source/Developer/AssetTools/Private/AssetTypeActions/AssetTypeActions_NiagaraScript.h
@@ -6,13 +6,13 @@ class FAssetTypeActions_NiagaraScript : public FAssetTypeActions_Base
 {
 public:
 	// IAssetTypeActions Implementation
-	virtual FText GetName() const OVERRIDE { return NSLOCTEXT("AssetTypeActions", "AssetTypeActions_NiagaraScript", "Niagara Script"); }
-	virtual FColor GetTypeColor() const OVERRIDE { return FColor(255,194,0); }
-	virtual UClass* GetSupportedClass() const OVERRIDE;
-	virtual bool HasActions ( const TArray<UObject*>& InObjects ) const OVERRIDE { return true; }
-	virtual void GetActions( const TArray<UObject*>& InObjects, FMenuBuilder& MenuBuilder ) OVERRIDE;
-	virtual void OpenAssetEditor( const TArray<UObject*>& InObjects, TSharedPtr<class IToolkitHost> EditWithinLevelEditor = TSharedPtr<IToolkitHost>() ) OVERRIDE;
-	virtual uint32 GetCategories() OVERRIDE { return EAssetTypeCategories::Misc; }
+	virtual FText GetName() const override { return NSLOCTEXT("AssetTypeActions", "AssetTypeActions_NiagaraScript", "Niagara Script"); }
+	virtual FColor GetTypeColor() const override { return FColor(255,194,0); }
+	virtual UClass* GetSupportedClass() const override;
+	virtual bool HasActions ( const TArray<UObject*>& InObjects ) const override { return true; }
+	virtual void GetActions( const TArray<UObject*>& InObjects, FMenuBuilder& MenuBuilder ) override;
+	virtual void OpenAssetEditor( const TArray<UObject*>& InObjects, TSharedPtr<class IToolkitHost> EditWithinLevelEditor = TSharedPtr<IToolkitHost>() ) override;
+	virtual uint32 GetCategories() override { return EAssetTypeCategories::Misc; }
 
 private:
 	/** Handler for when Edit is selected */
diff --git a/Engine/Source/Developer/AssetTools/Private/AssetTypeActions/AssetTypeActions_ObjectLibrary.h b/Engine/Source/Developer/AssetTools/Private/AssetTypeActions/AssetTypeActions_ObjectLibrary.h
--- a/Engine/Source/Developer/AssetTools/Private/AssetTypeActions/AssetTypeActions_ObjectLibrary.h
+++ b/Engine/Source/Developer/AssetTools/Private/AssetTypeActions/AssetTypeActions_ObjectLibrary.h
@@ -6,8 +6,8 @@ class FAssetTypeActions_ObjectLibrary : public FAssetTypeActions_Base
 {
 public:
 	// IAssetTypeActions Implementation
-	virtual FText GetName() const OVERRIDE { return NSLOCTEXT("AssetTypeActions", "AssetTypeActions_ObjectLibrary", "Object Library"); }
-	virtual FColor GetTypeColor() const OVERRIDE { return FColor(255, 63, 108); }
-	virtual UClass* GetSupportedClass() const OVERRIDE { return UObjectLibrary::StaticClass(); }
-	virtual uint32 GetCategories() OVERRIDE { return EAssetTypeCategories::Misc; }
+	virtual FText GetName() const override { return NSLOCTEXT("AssetTypeActions", "AssetTypeActions_ObjectLibrary", "Object Library"); }
+	virtual FColor GetTypeColor() const override { return FColor(255, 63, 108); }
+	virtual UClass* GetSupportedClass() const override { return UObjectLibrary::StaticClass(); }
+	virtual uint32 GetCategories() override { return EAssetTypeCategories::Misc; }
 };
\ No newline at end of file
diff --git a/Engine/Source/Developer/AssetTools/Private/AssetTypeActions/AssetTypeActions_ParticleSystem.h b/Engine/Source/Developer/AssetTools/Private/AssetTypeActions/AssetTypeActions_ParticleSystem.h
--- a/Engine/Source/Developer/AssetTools/Private/AssetTypeActions/AssetTypeActions_ParticleSystem.h
+++ b/Engine/Source/Developer/AssetTools/Private/AssetTypeActions/AssetTypeActions_ParticleSystem.h
@@ -8,13 +8,13 @@ class FAssetTypeActions_ParticleSystem : public FAssetTypeActions_Base
 {
 public:
 	// IAssetTypeActions Implementation
-	virtual FText GetName() const OVERRIDE { return NSLOCTEXT("AssetTypeActions", "AssetTypeActions_ParticleSystem", "Particle System"); }
-	virtual FColor GetTypeColor() const OVERRIDE { return FColor(255,255,255); }
-	virtual UClass* GetSupportedClass() const OVERRIDE { return UParticleSystem::StaticClass(); }
-	virtual bool HasActions ( const TArray<UObject*>& InObjects ) const OVERRIDE { return true; }
-	virtual void GetActions( const TArray<UObject*>& InObjects, FMenuBuilder& MenuBuilder ) OVERRIDE;
-	virtual void OpenAssetEditor( const TArray<UObject*>& InObjects, TSharedPtr<class IToolkitHost> EditWithinLevelEditor = TSharedPtr<IToolkitHost>() ) OVERRIDE;
-	virtual uint32 GetCategories() OVERRIDE { return EAssetTypeCategories::Basic; }
+	virtual FText GetName() const override { return NSLOCTEXT("AssetTypeActions", "AssetTypeActions_ParticleSystem", "Particle System"); }
+	virtual FColor GetTypeColor() const override { return FColor(255,255,255); }
+	virtual UClass* GetSupportedClass() const override { return UParticleSystem::StaticClass(); }
+	virtual bool HasActions ( const TArray<UObject*>& InObjects ) const override { return true; }
+	virtual void GetActions( const TArray<UObject*>& InObjects, FMenuBuilder& MenuBuilder ) override;
+	virtual void OpenAssetEditor( const TArray<UObject*>& InObjects, TSharedPtr<class IToolkitHost> EditWithinLevelEditor = TSharedPtr<IToolkitHost>() ) override;
+	virtual uint32 GetCategories() override { return EAssetTypeCategories::Basic; }
 	
 private:
 	/** Handler for when Edit is selected */
diff --git a/Engine/Source/Developer/AssetTools/Private/AssetTypeActions/AssetTypeActions_PhysicalMaterial.h b/Engine/Source/Developer/AssetTools/Private/AssetTypeActions/AssetTypeActions_PhysicalMaterial.h
--- a/Engine/Source/Developer/AssetTools/Private/AssetTypeActions/AssetTypeActions_PhysicalMaterial.h
+++ b/Engine/Source/Developer/AssetTools/Private/AssetTypeActions/AssetTypeActions_PhysicalMaterial.h
@@ -6,8 +6,8 @@ class FAssetTypeActions_PhysicalMaterial : public FAssetTypeActions_Base
 {
 public:
 	// IAssetTypeActions Implementation
-	virtual FText GetName() const OVERRIDE { return NSLOCTEXT("AssetTypeActions", "AssetTypeActions_PhysicalMaterial", "Physical Material"); }
-	virtual FColor GetTypeColor() const OVERRIDE { return FColor(200,192,128); }
-	virtual UClass* GetSupportedClass() const OVERRIDE;
-	virtual uint32 GetCategories() OVERRIDE { return EAssetTypeCategories::Physics; }
+	virtual FText GetName() const override { return NSLOCTEXT("AssetTypeActions", "AssetTypeActions_PhysicalMaterial", "Physical Material"); }
+	virtual FColor GetTypeColor() const override { return FColor(200,192,128); }
+	virtual UClass* GetSupportedClass() const override;
+	virtual uint32 GetCategories() override { return EAssetTypeCategories::Physics; }
 };
\ No newline at end of file
diff --git a/Engine/Source/Developer/AssetTools/Private/AssetTypeActions/AssetTypeActions_PhysicsAsset.h b/Engine/Source/Developer/AssetTools/Private/AssetTypeActions/AssetTypeActions_PhysicsAsset.h
--- a/Engine/Source/Developer/AssetTools/Private/AssetTypeActions/AssetTypeActions_PhysicsAsset.h
+++ b/Engine/Source/Developer/AssetTools/Private/AssetTypeActions/AssetTypeActions_PhysicsAsset.h
@@ -6,13 +6,13 @@ class FAssetTypeActions_PhysicsAsset : public FAssetTypeActions_Base
 {
 public:
 	// IAssetTypeActions Implementation
-	virtual FText GetName() const OVERRIDE { return NSLOCTEXT("AssetTypeActions", "AssetTypeActions_PhysicsAsset", "Physics Asset"); }
-	virtual FColor GetTypeColor() const OVERRIDE { return FColor(255,192,128); }
-	virtual UClass* GetSupportedClass() const OVERRIDE;
-	virtual bool HasActions ( const TArray<UObject*>& InObjects ) const OVERRIDE { return true; }
-	virtual void GetActions( const TArray<UObject*>& InObjects, FMenuBuilder& MenuBuilder ) OVERRIDE;
-	virtual void OpenAssetEditor( const TArray<UObject*>& InObjects, TSharedPtr<class IToolkitHost> EditWithinLevelEditor = TSharedPtr<IToolkitHost>() ) OVERRIDE;
-	virtual uint32 GetCategories() OVERRIDE { return EAssetTypeCategories::Physics; }
+	virtual FText GetName() const override { return NSLOCTEXT("AssetTypeActions", "AssetTypeActions_PhysicsAsset", "Physics Asset"); }
+	virtual FColor GetTypeColor() const override { return FColor(255,192,128); }
+	virtual UClass* GetSupportedClass() const override;
+	virtual bool HasActions ( const TArray<UObject*>& InObjects ) const override { return true; }
+	virtual void GetActions( const TArray<UObject*>& InObjects, FMenuBuilder& MenuBuilder ) override;
+	virtual void OpenAssetEditor( const TArray<UObject*>& InObjects, TSharedPtr<class IToolkitHost> EditWithinLevelEditor = TSharedPtr<IToolkitHost>() ) override;
+	virtual uint32 GetCategories() override { return EAssetTypeCategories::Physics; }
 
 private:
 	/** Handler for when Edit is selected */
diff --git a/Engine/Source/Developer/AssetTools/Private/AssetTypeActions/AssetTypeActions_Redirector.h b/Engine/Source/Developer/AssetTools/Private/AssetTypeActions/AssetTypeActions_Redirector.h
--- a/Engine/Source/Developer/AssetTools/Private/AssetTypeActions/AssetTypeActions_Redirector.h
+++ b/Engine/Source/Developer/AssetTools/Private/AssetTypeActions/AssetTypeActions_Redirector.h
@@ -6,13 +6,13 @@ class FAssetTypeActions_Redirector : public FAssetTypeActions_Base
 {
 public:
 	// IAssetTypeActions Implementation
-	virtual FText GetName() const OVERRIDE { return NSLOCTEXT("AssetTypeActions", "AssetTypeActions_Redirector", "Redirector"); }
-	virtual FColor GetTypeColor() const OVERRIDE { return FColor(128, 128, 128); }
-	virtual UClass* GetSupportedClass() const OVERRIDE { return UObjectRedirector::StaticClass(); }
-	virtual bool HasActions ( const TArray<UObject*>& InObjects ) const OVERRIDE { return true; }
-	virtual void GetActions( const TArray<UObject*>& InObjects, FMenuBuilder& MenuBuilder ) OVERRIDE;
-	virtual void AssetsActivated( const TArray<UObject*>& InObjects, EAssetTypeActivationMethod::Type ActivationType ) OVERRIDE;
-	virtual uint32 GetCategories() OVERRIDE { return EAssetTypeCategories::Misc; }
+	virtual FText GetName() const override { return NSLOCTEXT("AssetTypeActions", "AssetTypeActions_Redirector", "Redirector"); }
+	virtual FColor GetTypeColor() const override { return FColor(128, 128, 128); }
+	virtual UClass* GetSupportedClass() const override { return UObjectRedirector::StaticClass(); }
+	virtual bool HasActions ( const TArray<UObject*>& InObjects ) const override { return true; }
+	virtual void GetActions( const TArray<UObject*>& InObjects, FMenuBuilder& MenuBuilder ) override;
+	virtual void AssetsActivated( const TArray<UObject*>& InObjects, EAssetTypeActivationMethod::Type ActivationType ) override;
+	virtual uint32 GetCategories() override { return EAssetTypeCategories::Misc; }
 
 private:
 	/** Handler for when FindTarget is selected */
diff --git a/Engine/Source/Developer/AssetTools/Private/AssetTypeActions/AssetTypeActions_ReverbEffect.h b/Engine/Source/Developer/AssetTools/Private/AssetTypeActions/AssetTypeActions_ReverbEffect.h
--- a/Engine/Source/Developer/AssetTools/Private/AssetTypeActions/AssetTypeActions_ReverbEffect.h
+++ b/Engine/Source/Developer/AssetTools/Private/AssetTypeActions/AssetTypeActions_ReverbEffect.h
@@ -6,8 +6,8 @@ class FAssetTypeActions_ReverbEffect : public FAssetTypeActions_Base
 {
 public:
 	// IAssetTypeActions Implementation
-	virtual FText GetName() const OVERRIDE { return NSLOCTEXT("AssetTypeActions", "AssetTypeActions_ReverbEffect", "Reverb Effect"); }
-	virtual FColor GetTypeColor() const OVERRIDE { return FColor(175, 0, 0); }
-	virtual UClass* GetSupportedClass() const OVERRIDE;
-	virtual uint32 GetCategories() OVERRIDE { return EAssetTypeCategories::Sounds; }
+	virtual FText GetName() const override { return NSLOCTEXT("AssetTypeActions", "AssetTypeActions_ReverbEffect", "Reverb Effect"); }
+	virtual FColor GetTypeColor() const override { return FColor(175, 0, 0); }
+	virtual UClass* GetSupportedClass() const override;
+	virtual uint32 GetCategories() override { return EAssetTypeCategories::Sounds; }
 };
\ No newline at end of file
diff --git a/Engine/Source/Developer/AssetTools/Private/AssetTypeActions/AssetTypeActions_SkeletalMesh.h b/Engine/Source/Developer/AssetTools/Private/AssetTypeActions/AssetTypeActions_SkeletalMesh.h
--- a/Engine/Source/Developer/AssetTools/Private/AssetTypeActions/AssetTypeActions_SkeletalMesh.h
+++ b/Engine/Source/Developer/AssetTools/Private/AssetTypeActions/AssetTypeActions_SkeletalMesh.h
@@ -6,14 +6,14 @@ class FAssetTypeActions_SkeletalMesh : public FAssetTypeActions_Base
 {
 public:
 	// IAssetTypeActions Implementation
-	virtual FText GetName() const OVERRIDE { return NSLOCTEXT("AssetTypeActions", "AssetTypeActions_SkeletalMesh", "Skeletal Mesh"); }
-	virtual FColor GetTypeColor() const OVERRIDE { return FColor(255,0,255); }
-	virtual UClass* GetSupportedClass() const OVERRIDE { return USkeletalMesh::StaticClass(); }
-	virtual bool HasActions ( const TArray<UObject*>& InObjects ) const OVERRIDE { return true; }
-	virtual void GetActions( const TArray<UObject*>& InObjects, FMenuBuilder& MenuBuilder ) OVERRIDE;
-	virtual void OpenAssetEditor( const TArray<UObject*>& InObjects, TSharedPtr<class IToolkitHost> EditWithinLevelEditor = TSharedPtr<IToolkitHost>() ) OVERRIDE;
-	virtual uint32 GetCategories() OVERRIDE { return EAssetTypeCategories::Basic | EAssetTypeCategories::Animation; }
-	virtual class UThumbnailInfo* GetThumbnailInfo(UObject* Asset) const OVERRIDE;
+	virtual FText GetName() const override { return NSLOCTEXT("AssetTypeActions", "AssetTypeActions_SkeletalMesh", "Skeletal Mesh"); }
+	virtual FColor GetTypeColor() const override { return FColor(255,0,255); }
+	virtual UClass* GetSupportedClass() const override { return USkeletalMesh::StaticClass(); }
+	virtual bool HasActions ( const TArray<UObject*>& InObjects ) const override { return true; }
+	virtual void GetActions( const TArray<UObject*>& InObjects, FMenuBuilder& MenuBuilder ) override;
+	virtual void OpenAssetEditor( const TArray<UObject*>& InObjects, TSharedPtr<class IToolkitHost> EditWithinLevelEditor = TSharedPtr<IToolkitHost>() ) override;
+	virtual uint32 GetCategories() override { return EAssetTypeCategories::Basic | EAssetTypeCategories::Animation; }
+	virtual class UThumbnailInfo* GetThumbnailInfo(UObject* Asset) const override;
 
 protected:
 	/** Gets additional actions that do not apply to destructible meshes */
diff --git a/Engine/Source/Developer/AssetTools/Private/AssetTypeActions/AssetTypeActions_Skeleton.h b/Engine/Source/Developer/AssetTools/Private/AssetTypeActions/AssetTypeActions_Skeleton.h
--- a/Engine/Source/Developer/AssetTools/Private/AssetTypeActions/AssetTypeActions_Skeleton.h
+++ b/Engine/Source/Developer/AssetTools/Private/AssetTypeActions/AssetTypeActions_Skeleton.h
@@ -29,13 +29,13 @@ class FAssetTypeActions_Skeleton : public FAssetTypeActions_Base
 {
 public:
 	// IAssetTypeActions Implementation
-	virtual FText GetName() const OVERRIDE { return NSLOCTEXT("AssetTypeActions", "AssetTypeActions_Skeleton", "Skeleton"); }
-	virtual FColor GetTypeColor() const OVERRIDE { return FColor(105,181,205); }
-	virtual UClass* GetSupportedClass() const OVERRIDE { return USkeleton::StaticClass(); }
-	virtual bool HasActions ( const TArray<UObject*>& InObjects ) const OVERRIDE { return true; }
-	virtual void GetActions( const TArray<UObject*>& InObjects, FMenuBuilder& MenuBuilder ) OVERRIDE;
-	virtual void OpenAssetEditor( const TArray<UObject*>& InObjects, TSharedPtr<class IToolkitHost> EditWithinLevelEditor = TSharedPtr<IToolkitHost>() ) OVERRIDE;
-	virtual uint32 GetCategories() OVERRIDE { return EAssetTypeCategories::Animation; }
+	virtual FText GetName() const override { return NSLOCTEXT("AssetTypeActions", "AssetTypeActions_Skeleton", "Skeleton"); }
+	virtual FColor GetTypeColor() const override { return FColor(105,181,205); }
+	virtual UClass* GetSupportedClass() const override { return USkeleton::StaticClass(); }
+	virtual bool HasActions ( const TArray<UObject*>& InObjects ) const override { return true; }
+	virtual void GetActions( const TArray<UObject*>& InObjects, FMenuBuilder& MenuBuilder ) override;
+	virtual void OpenAssetEditor( const TArray<UObject*>& InObjects, TSharedPtr<class IToolkitHost> EditWithinLevelEditor = TSharedPtr<IToolkitHost>() ) override;
+	virtual uint32 GetCategories() override { return EAssetTypeCategories::Animation; }
 
 	/** Handler for the blend space sub menu */
 	void FillBlendSpaceMenu( FMenuBuilder& MenuBuilder, TArray<TWeakObjectPtr<USkeleton>> Skeletons );
diff --git a/Engine/Source/Developer/AssetTools/Private/AssetTypeActions/AssetTypeActions_SlateBrush.h b/Engine/Source/Developer/AssetTools/Private/AssetTypeActions/AssetTypeActions_SlateBrush.h
--- a/Engine/Source/Developer/AssetTools/Private/AssetTypeActions/AssetTypeActions_SlateBrush.h
+++ b/Engine/Source/Developer/AssetTools/Private/AssetTypeActions/AssetTypeActions_SlateBrush.h
@@ -8,9 +8,9 @@ class FAssetTypeActions_SlateBrush : public FAssetTypeActions_Base
 {
 public:
 	// IAssetTypeActions Implementation
-	virtual FText GetName() const OVERRIDE { return NSLOCTEXT("AssetTypeActions", "AssetTypeActions_SlateBrush", "Slate Brush"); }
-	virtual FColor GetTypeColor() const OVERRIDE { return FColor(105, 165, 60); }
-	virtual UClass* GetSupportedClass() const OVERRIDE { return USlateBrushAsset::StaticClass(); }
-	virtual uint32 GetCategories() OVERRIDE { return EAssetTypeCategories::UI; }
-	virtual bool HasActions ( const TArray<UObject*>& InObjects ) const OVERRIDE { return false; }
+	virtual FText GetName() const override { return NSLOCTEXT("AssetTypeActions", "AssetTypeActions_SlateBrush", "Slate Brush"); }
+	virtual FColor GetTypeColor() const override { return FColor(105, 165, 60); }
+	virtual UClass* GetSupportedClass() const override { return USlateBrushAsset::StaticClass(); }
+	virtual uint32 GetCategories() override { return EAssetTypeCategories::UI; }
+	virtual bool HasActions ( const TArray<UObject*>& InObjects ) const override { return false; }
 };
\ No newline at end of file
diff --git a/Engine/Source/Developer/AssetTools/Private/AssetTypeActions/AssetTypeActions_SlateWidgetStyle.h b/Engine/Source/Developer/AssetTools/Private/AssetTypeActions/AssetTypeActions_SlateWidgetStyle.h
--- a/Engine/Source/Developer/AssetTools/Private/AssetTypeActions/AssetTypeActions_SlateWidgetStyle.h
+++ b/Engine/Source/Developer/AssetTools/Private/AssetTypeActions/AssetTypeActions_SlateWidgetStyle.h
@@ -6,14 +6,14 @@ class FAssetTypeActions_SlateWidgetStyle : public FAssetTypeActions_Base
 {
 public:
 	// IAssetTypeActions Implementation
-	virtual FText GetName() const OVERRIDE { return NSLOCTEXT("AssetTypeActions", "AssetTypeActions_SlateStyle", "Slate Widget Style"); }
-	virtual FColor GetTypeColor() const OVERRIDE { return FColor(62, 140, 35); }
-	virtual UClass* GetSupportedClass() const OVERRIDE { return USlateWidgetStyleAsset::StaticClass(); }
-	virtual uint32 GetCategories() OVERRIDE { return EAssetTypeCategories::UI; }
-	virtual bool HasActions ( const TArray<UObject*>& InObjects ) const OVERRIDE { return true; }
-	virtual void GetActions( const TArray<UObject*>& InObjects, FMenuBuilder& MenuBuilder ) OVERRIDE;
+	virtual FText GetName() const override { return NSLOCTEXT("AssetTypeActions", "AssetTypeActions_SlateStyle", "Slate Widget Style"); }
+	virtual FColor GetTypeColor() const override { return FColor(62, 140, 35); }
+	virtual UClass* GetSupportedClass() const override { return USlateWidgetStyleAsset::StaticClass(); }
+	virtual uint32 GetCategories() override { return EAssetTypeCategories::UI; }
+	virtual bool HasActions ( const TArray<UObject*>& InObjects ) const override { return true; }
+	virtual void GetActions( const TArray<UObject*>& InObjects, FMenuBuilder& MenuBuilder ) override;
 
-	virtual void OpenAssetEditor( const TArray<UObject*>& InObjects, TSharedPtr<class IToolkitHost> EditWithinLevelEditor = TSharedPtr<IToolkitHost>() ) OVERRIDE;
+	virtual void OpenAssetEditor( const TArray<UObject*>& InObjects, TSharedPtr<class IToolkitHost> EditWithinLevelEditor = TSharedPtr<IToolkitHost>() ) override;
 
 private:
 	/** Handler for when Edit is selected */
diff --git a/Engine/Source/Developer/AssetTools/Private/AssetTypeActions/AssetTypeActions_SoundAttenuation.h b/Engine/Source/Developer/AssetTools/Private/AssetTypeActions/AssetTypeActions_SoundAttenuation.h
--- a/Engine/Source/Developer/AssetTools/Private/AssetTypeActions/AssetTypeActions_SoundAttenuation.h
+++ b/Engine/Source/Developer/AssetTools/Private/AssetTypeActions/AssetTypeActions_SoundAttenuation.h
@@ -6,8 +6,8 @@ class FAssetTypeActions_SoundAttenuation : public FAssetTypeActions_Base
 {
 public:
 	// IAssetTypeActions Implementation
-	virtual FText GetName() const OVERRIDE { return NSLOCTEXT("AssetTypeActions", "AssetTypeActions_SoundAttenuation", "Sound Attenuation"); }
-	virtual FColor GetTypeColor() const OVERRIDE { return FColor(0, 0, 175); }
-	virtual UClass* GetSupportedClass() const OVERRIDE;
-	virtual uint32 GetCategories() OVERRIDE { return EAssetTypeCategories::Sounds; }
+	virtual FText GetName() const override { return NSLOCTEXT("AssetTypeActions", "AssetTypeActions_SoundAttenuation", "Sound Attenuation"); }
+	virtual FColor GetTypeColor() const override { return FColor(0, 0, 175); }
+	virtual UClass* GetSupportedClass() const override;
+	virtual uint32 GetCategories() override { return EAssetTypeCategories::Sounds; }
 };
\ No newline at end of file
diff --git a/Engine/Source/Developer/AssetTools/Private/AssetTypeActions/AssetTypeActions_SoundBase.h b/Engine/Source/Developer/AssetTools/Private/AssetTypeActions/AssetTypeActions_SoundBase.h
--- a/Engine/Source/Developer/AssetTools/Private/AssetTypeActions/AssetTypeActions_SoundBase.h
+++ b/Engine/Source/Developer/AssetTools/Private/AssetTypeActions/AssetTypeActions_SoundBase.h
@@ -6,14 +6,14 @@ class FAssetTypeActions_SoundBase : public FAssetTypeActions_Base
 {
 public:
 	// IAssetTypeActions Implementation
-	virtual FText GetName() const OVERRIDE { return NSLOCTEXT("AssetTypeActions", "AssetTypeActions_SoundBase", "Sound Base"); }
-	virtual FColor GetTypeColor() const OVERRIDE { return FColor(0,0,255); }
-	virtual UClass* GetSupportedClass() const OVERRIDE;
-	virtual bool HasActions ( const TArray<UObject*>& InObjects ) const OVERRIDE { return true; }
-	virtual void GetActions( const TArray<UObject*>& InObjects, FMenuBuilder& MenuBuilder ) OVERRIDE;
-	virtual void AssetsActivated( const TArray<UObject*>& InObjects, EAssetTypeActivationMethod::Type ActivationType ) OVERRIDE;
-	virtual uint32 GetCategories() OVERRIDE { return EAssetTypeCategories::Sounds; }
-	virtual bool CanFilter() OVERRIDE { return false; }
+	virtual FText GetName() const override { return NSLOCTEXT("AssetTypeActions", "AssetTypeActions_SoundBase", "Sound Base"); }
+	virtual FColor GetTypeColor() const override { return FColor(0,0,255); }
+	virtual UClass* GetSupportedClass() const override;
+	virtual bool HasActions ( const TArray<UObject*>& InObjects ) const override { return true; }
+	virtual void GetActions( const TArray<UObject*>& InObjects, FMenuBuilder& MenuBuilder ) override;
+	virtual void AssetsActivated( const TArray<UObject*>& InObjects, EAssetTypeActivationMethod::Type ActivationType ) override;
+	virtual uint32 GetCategories() override { return EAssetTypeCategories::Sounds; }
+	virtual bool CanFilter() override { return false; }
 
 private:
 	/** Handler for when PlaySound is selected */
diff --git a/Engine/Source/Developer/AssetTools/Private/AssetTypeActions/AssetTypeActions_SoundClass.h b/Engine/Source/Developer/AssetTools/Private/AssetTypeActions/AssetTypeActions_SoundClass.h
--- a/Engine/Source/Developer/AssetTools/Private/AssetTypeActions/AssetTypeActions_SoundClass.h
+++ b/Engine/Source/Developer/AssetTools/Private/AssetTypeActions/AssetTypeActions_SoundClass.h
@@ -6,13 +6,13 @@ class FAssetTypeActions_SoundClass : public FAssetTypeActions_Base
 {
 public:
 	// IAssetTypeActions Implementation
-	virtual FText GetName() const OVERRIDE { return NSLOCTEXT("AssetTypeActions", "AssetTypeActions_SoundClass", "Sound Class"); }
-	virtual FColor GetTypeColor() const OVERRIDE { return FColor(255, 175, 0); }
-	virtual UClass* GetSupportedClass() const OVERRIDE;
-	virtual bool HasActions ( const TArray<UObject*>& InObjects ) const OVERRIDE { return true; }
-	virtual void GetActions( const TArray<UObject*>& InObjects, FMenuBuilder& MenuBuilder ) OVERRIDE;
-	virtual void OpenAssetEditor( const TArray<UObject*>& InObjects, TSharedPtr<class IToolkitHost> EditWithinLevelEditor = TSharedPtr<IToolkitHost>() ) OVERRIDE;
-	virtual uint32 GetCategories() OVERRIDE { return EAssetTypeCategories::Sounds; }
+	virtual FText GetName() const override { return NSLOCTEXT("AssetTypeActions", "AssetTypeActions_SoundClass", "Sound Class"); }
+	virtual FColor GetTypeColor() const override { return FColor(255, 175, 0); }
+	virtual UClass* GetSupportedClass() const override;
+	virtual bool HasActions ( const TArray<UObject*>& InObjects ) const override { return true; }
+	virtual void GetActions( const TArray<UObject*>& InObjects, FMenuBuilder& MenuBuilder ) override;
+	virtual void OpenAssetEditor( const TArray<UObject*>& InObjects, TSharedPtr<class IToolkitHost> EditWithinLevelEditor = TSharedPtr<IToolkitHost>() ) override;
+	virtual uint32 GetCategories() override { return EAssetTypeCategories::Sounds; }
 
 private:
 	/** Handler for when Edit is selected */
diff --git a/Engine/Source/Developer/AssetTools/Private/AssetTypeActions/AssetTypeActions_SoundCue.h b/Engine/Source/Developer/AssetTools/Private/AssetTypeActions/AssetTypeActions_SoundCue.h
--- a/Engine/Source/Developer/AssetTools/Private/AssetTypeActions/AssetTypeActions_SoundCue.h
+++ b/Engine/Source/Developer/AssetTools/Private/AssetTypeActions/AssetTypeActions_SoundCue.h
@@ -6,12 +6,12 @@ class FAssetTypeActions_SoundCue : public FAssetTypeActions_SoundBase
 {
 public:
 	// IAssetTypeActions Implementation
-	virtual FText GetName() const OVERRIDE { return NSLOCTEXT("AssetTypeActions", "AssetTypeActions_SoundCue", "Sound Cue"); }
-	virtual FColor GetTypeColor() const OVERRIDE { return FColor(0, 175, 255); }
-	virtual UClass* GetSupportedClass() const OVERRIDE;
-	virtual void GetActions( const TArray<UObject*>& InObjects, FMenuBuilder& MenuBuilder ) OVERRIDE;
-	virtual void OpenAssetEditor( const TArray<UObject*>& InObjects, TSharedPtr<class IToolkitHost> EditWithinLevelEditor = TSharedPtr<IToolkitHost>() ) OVERRIDE;
-	virtual bool CanFilter() OVERRIDE { return true; }
+	virtual FText GetName() const override { return NSLOCTEXT("AssetTypeActions", "AssetTypeActions_SoundCue", "Sound Cue"); }
+	virtual FColor GetTypeColor() const override { return FColor(0, 175, 255); }
+	virtual UClass* GetSupportedClass() const override;
+	virtual void GetActions( const TArray<UObject*>& InObjects, FMenuBuilder& MenuBuilder ) override;
+	virtual void OpenAssetEditor( const TArray<UObject*>& InObjects, TSharedPtr<class IToolkitHost> EditWithinLevelEditor = TSharedPtr<IToolkitHost>() ) override;
+	virtual bool CanFilter() override { return true; }
 
 private:
 	/** Handler for when Edit is selected */
diff --git a/Engine/Source/Developer/AssetTools/Private/AssetTypeActions/AssetTypeActions_SoundMix.h b/Engine/Source/Developer/AssetTools/Private/AssetTypeActions/AssetTypeActions_SoundMix.h
--- a/Engine/Source/Developer/AssetTools/Private/AssetTypeActions/AssetTypeActions_SoundMix.h
+++ b/Engine/Source/Developer/AssetTools/Private/AssetTypeActions/AssetTypeActions_SoundMix.h
@@ -6,8 +6,8 @@ class FAssetTypeActions_SoundMix : public FAssetTypeActions_Base
 {
 public:
 	// IAssetTypeActions Implementation
-	virtual FText GetName() const OVERRIDE { return NSLOCTEXT("AssetTypeActions", "AssetTypeActions_SoundMix", "Sound Mix"); }
-	virtual FColor GetTypeColor() const OVERRIDE { return FColor(175, 0, 255); }
-	virtual UClass* GetSupportedClass() const OVERRIDE;
-	virtual uint32 GetCategories() OVERRIDE { return EAssetTypeCategories::Sounds; }
+	virtual FText GetName() const override { return NSLOCTEXT("AssetTypeActions", "AssetTypeActions_SoundMix", "Sound Mix"); }
+	virtual FColor GetTypeColor() const override { return FColor(175, 0, 255); }
+	virtual UClass* GetSupportedClass() const override;
+	virtual uint32 GetCategories() override { return EAssetTypeCategories::Sounds; }
 };
\ No newline at end of file
diff --git a/Engine/Source/Developer/AssetTools/Private/AssetTypeActions/AssetTypeActions_SoundWave.h b/Engine/Source/Developer/AssetTools/Private/AssetTypeActions/AssetTypeActions_SoundWave.h
--- a/Engine/Source/Developer/AssetTools/Private/AssetTypeActions/AssetTypeActions_SoundWave.h
+++ b/Engine/Source/Developer/AssetTools/Private/AssetTypeActions/AssetTypeActions_SoundWave.h
@@ -6,12 +6,12 @@ class FAssetTypeActions_SoundWave : public FAssetTypeActions_SoundBase
 {
 public:
 	// IAssetTypeActions Implementation
-	virtual FText GetName() const OVERRIDE { return NSLOCTEXT("AssetTypeActions", "AssetTypeActions_SoundWave", "Sound Wave"); }
-	virtual FColor GetTypeColor() const OVERRIDE { return FColor(0, 0, 255); }
-	virtual UClass* GetSupportedClass() const OVERRIDE;
-	virtual void GetActions( const TArray<UObject*>& InObjects, FMenuBuilder& MenuBuilder ) OVERRIDE;
-	virtual void OpenAssetEditor( const TArray<UObject*>& InObjects, TSharedPtr<class IToolkitHost> EditWithinLevelEditor = TSharedPtr<IToolkitHost>() ) OVERRIDE;
-	virtual bool CanFilter() OVERRIDE { return true; }
+	virtual FText GetName() const override { return NSLOCTEXT("AssetTypeActions", "AssetTypeActions_SoundWave", "Sound Wave"); }
+	virtual FColor GetTypeColor() const override { return FColor(0, 0, 255); }
+	virtual UClass* GetSupportedClass() const override;
+	virtual void GetActions( const TArray<UObject*>& InObjects, FMenuBuilder& MenuBuilder ) override;
+	virtual void OpenAssetEditor( const TArray<UObject*>& InObjects, TSharedPtr<class IToolkitHost> EditWithinLevelEditor = TSharedPtr<IToolkitHost>() ) override;
+	virtual bool CanFilter() override { return true; }
 
 private:
 
diff --git a/Engine/Source/Developer/AssetTools/Private/AssetTypeActions/AssetTypeActions_StaticMesh.h b/Engine/Source/Developer/AssetTools/Private/AssetTypeActions/AssetTypeActions_StaticMesh.h
--- a/Engine/Source/Developer/AssetTools/Private/AssetTypeActions/AssetTypeActions_StaticMesh.h
+++ b/Engine/Source/Developer/AssetTools/Private/AssetTypeActions/AssetTypeActions_StaticMesh.h
@@ -6,14 +6,14 @@ class FAssetTypeActions_StaticMesh : public FAssetTypeActions_Base
 {
 public:
 	// IAssetTypeActions Implementation
-	virtual FText GetName() const OVERRIDE { return NSLOCTEXT("AssetTypeActions", "AssetTypeActions_StaticMesh", "Static Mesh"); }
-	virtual FColor GetTypeColor() const OVERRIDE { return FColor(0, 255, 255); }
-	virtual UClass* GetSupportedClass() const OVERRIDE { return UStaticMesh::StaticClass(); }
-	virtual bool HasActions ( const TArray<UObject*>& InObjects ) const OVERRIDE { return true; }
-	virtual void GetActions( const TArray<UObject*>& InObjects, FMenuBuilder& MenuBuilder ) OVERRIDE;
-	virtual void OpenAssetEditor( const TArray<UObject*>& InObjects, TSharedPtr<class IToolkitHost> EditWithinLevelEditor = TSharedPtr<IToolkitHost>() ) OVERRIDE;
-	virtual uint32 GetCategories() OVERRIDE { return EAssetTypeCategories::Basic; }
-	virtual class UThumbnailInfo* GetThumbnailInfo(UObject* Asset) const OVERRIDE;
+	virtual FText GetName() const override { return NSLOCTEXT("AssetTypeActions", "AssetTypeActions_StaticMesh", "Static Mesh"); }
+	virtual FColor GetTypeColor() const override { return FColor(0, 255, 255); }
+	virtual UClass* GetSupportedClass() const override { return UStaticMesh::StaticClass(); }
+	virtual bool HasActions ( const TArray<UObject*>& InObjects ) const override { return true; }
+	virtual void GetActions( const TArray<UObject*>& InObjects, FMenuBuilder& MenuBuilder ) override;
+	virtual void OpenAssetEditor( const TArray<UObject*>& InObjects, TSharedPtr<class IToolkitHost> EditWithinLevelEditor = TSharedPtr<IToolkitHost>() ) override;
+	virtual uint32 GetCategories() override { return EAssetTypeCategories::Basic; }
+	virtual class UThumbnailInfo* GetThumbnailInfo(UObject* Asset) const override;
 
 private:
 	/** Handler for when Edit is selected */
diff --git a/Engine/Source/Developer/AssetTools/Private/AssetTypeActions/AssetTypeActions_Struct.h b/Engine/Source/Developer/AssetTools/Private/AssetTypeActions/AssetTypeActions_Struct.h
--- a/Engine/Source/Developer/AssetTools/Private/AssetTypeActions/AssetTypeActions_Struct.h
+++ b/Engine/Source/Developer/AssetTools/Private/AssetTypeActions/AssetTypeActions_Struct.h
@@ -6,11 +6,11 @@ class FAssetTypeActions_Struct : public FAssetTypeActions_Base
 {
 public:
 	// IAssetTypeActions Implementation
-	virtual FText GetName() const OVERRIDE { return NSLOCTEXT("AssetTypeActions", "AssetTypeActions_Struct", "Structure"); }
-	virtual FColor GetTypeColor() const OVERRIDE { return FColor(0, 0, 255); }
-	virtual UClass* GetSupportedClass() const OVERRIDE { return UUserDefinedStruct::StaticClass(); }
-	virtual uint32 GetCategories() OVERRIDE { return EAssetTypeCategories::Misc; }
+	virtual FText GetName() const override { return NSLOCTEXT("AssetTypeActions", "AssetTypeActions_Struct", "Structure"); }
+	virtual FColor GetTypeColor() const override { return FColor(0, 0, 255); }
+	virtual UClass* GetSupportedClass() const override { return UUserDefinedStruct::StaticClass(); }
+	virtual uint32 GetCategories() override { return EAssetTypeCategories::Misc; }
 
-	virtual FText GetAssetDescription(const FAssetData& AssetData) const OVERRIDE;
-	virtual void OpenAssetEditor( const TArray<UObject*>& InObjects, TSharedPtr<class IToolkitHost> EditWithinLevelEditor = TSharedPtr<IToolkitHost>() ) OVERRIDE;
+	virtual FText GetAssetDescription(const FAssetData& AssetData) const override;
+	virtual void OpenAssetEditor( const TArray<UObject*>& InObjects, TSharedPtr<class IToolkitHost> EditWithinLevelEditor = TSharedPtr<IToolkitHost>() ) override;
 };
diff --git a/Engine/Source/Developer/AssetTools/Private/AssetTypeActions/AssetTypeActions_Texture.h b/Engine/Source/Developer/AssetTools/Private/AssetTypeActions/AssetTypeActions_Texture.h
--- a/Engine/Source/Developer/AssetTools/Private/AssetTypeActions/AssetTypeActions_Texture.h
+++ b/Engine/Source/Developer/AssetTools/Private/AssetTypeActions/AssetTypeActions_Texture.h
@@ -6,14 +6,14 @@ class FAssetTypeActions_Texture : public FAssetTypeActions_Base
 {
 public:
 	// IAssetTypeActions Implementation
-	virtual FText GetName() const OVERRIDE { return NSLOCTEXT("AssetTypeActions", "AssetTypeActions_Texture", "BaseTexture"); }
-	virtual FColor GetTypeColor() const OVERRIDE { return FColor(255,0,0); }
-	virtual UClass* GetSupportedClass() const OVERRIDE { return UTexture::StaticClass(); }
-	virtual bool HasActions ( const TArray<UObject*>& InObjects ) const OVERRIDE { return true; }
-	virtual void GetActions( const TArray<UObject*>& InObjects, FMenuBuilder& MenuBuilder ) OVERRIDE;
-	virtual void OpenAssetEditor( const TArray<UObject*>& InObjects, TSharedPtr<class IToolkitHost> EditWithinLevelEditor = TSharedPtr<IToolkitHost>() ) OVERRIDE;
-	virtual bool CanFilter() OVERRIDE { return false; }
-	virtual uint32 GetCategories() OVERRIDE { return EAssetTypeCategories::MaterialsAndTextures; }
+	virtual FText GetName() const override { return NSLOCTEXT("AssetTypeActions", "AssetTypeActions_Texture", "BaseTexture"); }
+	virtual FColor GetTypeColor() const override { return FColor(255,0,0); }
+	virtual UClass* GetSupportedClass() const override { return UTexture::StaticClass(); }
+	virtual bool HasActions ( const TArray<UObject*>& InObjects ) const override { return true; }
+	virtual void GetActions( const TArray<UObject*>& InObjects, FMenuBuilder& MenuBuilder ) override;
+	virtual void OpenAssetEditor( const TArray<UObject*>& InObjects, TSharedPtr<class IToolkitHost> EditWithinLevelEditor = TSharedPtr<IToolkitHost>() ) override;
+	virtual bool CanFilter() override { return false; }
+	virtual uint32 GetCategories() override { return EAssetTypeCategories::MaterialsAndTextures; }
 
 private:
 	/** Handler for when Reimport is selected */
diff --git a/Engine/Source/Developer/AssetTools/Private/AssetTypeActions/AssetTypeActions_Texture2D.h b/Engine/Source/Developer/AssetTools/Private/AssetTypeActions/AssetTypeActions_Texture2D.h
--- a/Engine/Source/Developer/AssetTools/Private/AssetTypeActions/AssetTypeActions_Texture2D.h
+++ b/Engine/Source/Developer/AssetTools/Private/AssetTypeActions/AssetTypeActions_Texture2D.h
@@ -6,13 +6,13 @@ class FAssetTypeActions_Texture2D : public FAssetTypeActions_Texture
 {
 public:
 	// IAssetTypeActions Implementation
-	virtual FText GetName() const OVERRIDE { return NSLOCTEXT("AssetTypeActions", "AssetTypeActions_Texture2D", "Texture"); }
-	virtual FColor GetTypeColor() const OVERRIDE { return FColor(255,0,0); }
-	virtual UClass* GetSupportedClass() const OVERRIDE { return UTexture2D::StaticClass(); }
-	virtual bool HasActions ( const TArray<UObject*>& InObjects ) const OVERRIDE { return true; }
-	virtual void GetActions( const TArray<UObject*>& InObjects, FMenuBuilder& MenuBuilder ) OVERRIDE;
-	virtual bool CanFilter() OVERRIDE { return true; }
-	virtual uint32 GetCategories() OVERRIDE { return FAssetTypeActions_Texture::GetCategories() | EAssetTypeCategories::Basic; }
+	virtual FText GetName() const override { return NSLOCTEXT("AssetTypeActions", "AssetTypeActions_Texture2D", "Texture"); }
+	virtual FColor GetTypeColor() const override { return FColor(255,0,0); }
+	virtual UClass* GetSupportedClass() const override { return UTexture2D::StaticClass(); }
+	virtual bool HasActions ( const TArray<UObject*>& InObjects ) const override { return true; }
+	virtual void GetActions( const TArray<UObject*>& InObjects, FMenuBuilder& MenuBuilder ) override;
+	virtual bool CanFilter() override { return true; }
+	virtual uint32 GetCategories() override { return FAssetTypeActions_Texture::GetCategories() | EAssetTypeCategories::Basic; }
 
 private:
 	/** Handler for when Create Slate Brush is selected */
diff --git a/Engine/Source/Developer/AssetTools/Private/AssetTypeActions/AssetTypeActions_TextureCube.h b/Engine/Source/Developer/AssetTools/Private/AssetTypeActions/AssetTypeActions_TextureCube.h
--- a/Engine/Source/Developer/AssetTools/Private/AssetTypeActions/AssetTypeActions_TextureCube.h
+++ b/Engine/Source/Developer/AssetTools/Private/AssetTypeActions/AssetTypeActions_TextureCube.h
@@ -6,8 +6,8 @@ class FAssetTypeActions_TextureCube : public FAssetTypeActions_Texture
 {
 public:
 	// IAssetTypeActions Implementation
-	virtual FText GetName() const OVERRIDE { return NSLOCTEXT("AssetTypeActions", "AssetTypeActions_TextureCube", "Texture Cube"); }
-	virtual FColor GetTypeColor() const OVERRIDE { return FColor(255,0,0); }
-	virtual UClass* GetSupportedClass() const OVERRIDE { return UTextureCube::StaticClass(); }
-	virtual bool CanFilter() OVERRIDE { return true; }
+	virtual FText GetName() const override { return NSLOCTEXT("AssetTypeActions", "AssetTypeActions_TextureCube", "Texture Cube"); }
+	virtual FColor GetTypeColor() const override { return FColor(255,0,0); }
+	virtual UClass* GetSupportedClass() const override { return UTextureCube::StaticClass(); }
+	virtual bool CanFilter() override { return true; }
 };
\ No newline at end of file
diff --git a/Engine/Source/Developer/AssetTools/Private/AssetTypeActions/AssetTypeActions_TextureLightProfile.h b/Engine/Source/Developer/AssetTools/Private/AssetTypeActions/AssetTypeActions_TextureLightProfile.h
--- a/Engine/Source/Developer/AssetTools/Private/AssetTypeActions/AssetTypeActions_TextureLightProfile.h
+++ b/Engine/Source/Developer/AssetTools/Private/AssetTypeActions/AssetTypeActions_TextureLightProfile.h
@@ -6,8 +6,8 @@ class FAssetTypeActions_TextureLightProfile : public FAssetTypeActions_Texture
 {
 public:
 	// IAssetTypeActions Implementation
-	virtual FText GetName() const OVERRIDE { return NSLOCTEXT("AssetTypeActions", "AssetTypeActions_TextureLightProfile", "Texture Light Profile"); }
-	virtual FColor GetTypeColor() const OVERRIDE { return FColor(255,0,0); }
-	virtual UClass* GetSupportedClass() const OVERRIDE { return UTextureLightProfile::StaticClass(); }
-	virtual bool CanFilter() OVERRIDE { return true; }
+	virtual FText GetName() const override { return NSLOCTEXT("AssetTypeActions", "AssetTypeActions_TextureLightProfile", "Texture Light Profile"); }
+	virtual FColor GetTypeColor() const override { return FColor(255,0,0); }
+	virtual UClass* GetSupportedClass() const override { return UTextureLightProfile::StaticClass(); }
+	virtual bool CanFilter() override { return true; }
 };
\ No newline at end of file
diff --git a/Engine/Source/Developer/AssetTools/Private/AssetTypeActions/AssetTypeActions_TextureMovie.h b/Engine/Source/Developer/AssetTools/Private/AssetTypeActions/AssetTypeActions_TextureMovie.h
--- a/Engine/Source/Developer/AssetTools/Private/AssetTypeActions/AssetTypeActions_TextureMovie.h
+++ b/Engine/Source/Developer/AssetTools/Private/AssetTypeActions/AssetTypeActions_TextureMovie.h
@@ -6,8 +6,8 @@ class FAssetTypeActions_TextureMovie : public FAssetTypeActions_Texture
 {
 public:
 	// IAssetTypeActions Implementation
-	virtual FText GetName() const OVERRIDE { return NSLOCTEXT("AssetTypeActions", "AssetTypeActions_TextureMovie", "Texture Movie"); }
-	virtual FColor GetTypeColor() const OVERRIDE { return FColor(255,0,0); }
-	virtual UClass* GetSupportedClass() const OVERRIDE { return UTextureMovie::StaticClass(); }
-	virtual bool CanFilter() OVERRIDE { return true; }
+	virtual FText GetName() const override { return NSLOCTEXT("AssetTypeActions", "AssetTypeActions_TextureMovie", "Texture Movie"); }
+	virtual FColor GetTypeColor() const override { return FColor(255,0,0); }
+	virtual UClass* GetSupportedClass() const override { return UTextureMovie::StaticClass(); }
+	virtual bool CanFilter() override { return true; }
 };
\ No newline at end of file
diff --git a/Engine/Source/Developer/AssetTools/Private/AssetTypeActions/AssetTypeActions_TextureRenderTarget.h b/Engine/Source/Developer/AssetTools/Private/AssetTypeActions/AssetTypeActions_TextureRenderTarget.h
--- a/Engine/Source/Developer/AssetTools/Private/AssetTypeActions/AssetTypeActions_TextureRenderTarget.h
+++ b/Engine/Source/Developer/AssetTools/Private/AssetTypeActions/AssetTypeActions_TextureRenderTarget.h
@@ -6,11 +6,11 @@ class FAssetTypeActions_TextureRenderTarget : public FAssetTypeActions_Texture
 {
 public:
 	// IAssetTypeActions Implementation
-	virtual FText GetName() const OVERRIDE { return NSLOCTEXT("AssetTypeActions", "AssetTypeActions_TextureRenderTarget", "Texture Render Target"); }
-	virtual FColor GetTypeColor() const OVERRIDE { return FColor(255,0,0); }
-	virtual UClass* GetSupportedClass() const OVERRIDE { return UTextureRenderTarget::StaticClass(); }
-	virtual bool HasActions ( const TArray<UObject*>& InObjects ) const OVERRIDE { return true; }
-	virtual void GetActions( const TArray<UObject*>& InObjects, FMenuBuilder& MenuBuilder ) OVERRIDE;
+	virtual FText GetName() const override { return NSLOCTEXT("AssetTypeActions", "AssetTypeActions_TextureRenderTarget", "Texture Render Target"); }
+	virtual FColor GetTypeColor() const override { return FColor(255,0,0); }
+	virtual UClass* GetSupportedClass() const override { return UTextureRenderTarget::StaticClass(); }
+	virtual bool HasActions ( const TArray<UObject*>& InObjects ) const override { return true; }
+	virtual void GetActions( const TArray<UObject*>& InObjects, FMenuBuilder& MenuBuilder ) override;
 
 private:
 	/** Handler for when CreateStatic is selected */
diff --git a/Engine/Source/Developer/AssetTools/Private/AssetTypeActions/AssetTypeActions_TextureRenderTarget2D.h b/Engine/Source/Developer/AssetTools/Private/AssetTypeActions/AssetTypeActions_TextureRenderTarget2D.h
--- a/Engine/Source/Developer/AssetTools/Private/AssetTypeActions/AssetTypeActions_TextureRenderTarget2D.h
+++ b/Engine/Source/Developer/AssetTools/Private/AssetTypeActions/AssetTypeActions_TextureRenderTarget2D.h
@@ -6,7 +6,7 @@ class FAssetTypeActions_TextureRenderTarget2D : public FAssetTypeActions_Texture
 {
 public:
 	// IAssetTypeActions Implementation
-	virtual FText GetName() const OVERRIDE { return NSLOCTEXT("AssetTypeActions", "AssetTypeActions_TextureRenderTarget2D", "Render Target"); }
-	virtual UClass* GetSupportedClass() const OVERRIDE { return UTextureRenderTarget2D::StaticClass(); }
-	virtual bool CanFilter() OVERRIDE { return true; }
+	virtual FText GetName() const override { return NSLOCTEXT("AssetTypeActions", "AssetTypeActions_TextureRenderTarget2D", "Render Target"); }
+	virtual UClass* GetSupportedClass() const override { return UTextureRenderTarget2D::StaticClass(); }
+	virtual bool CanFilter() override { return true; }
 };
\ No newline at end of file
diff --git a/Engine/Source/Developer/AssetTools/Private/AssetTypeActions/AssetTypeActions_TextureRenderTargetCube.h b/Engine/Source/Developer/AssetTools/Private/AssetTypeActions/AssetTypeActions_TextureRenderTargetCube.h
--- a/Engine/Source/Developer/AssetTools/Private/AssetTypeActions/AssetTypeActions_TextureRenderTargetCube.h
+++ b/Engine/Source/Developer/AssetTools/Private/AssetTypeActions/AssetTypeActions_TextureRenderTargetCube.h
@@ -6,7 +6,7 @@ class FAssetTypeActions_TextureRenderTargetCube : public FAssetTypeActions_Textu
 {
 public:
 	// IAssetTypeActions Implementation
-	virtual FText GetName() const OVERRIDE { return NSLOCTEXT("AssetTypeActions", "AssetTypeActions_TextureRenderTargetCube", "Cube Render Target"); }
-	virtual UClass* GetSupportedClass() const OVERRIDE { return UTextureRenderTargetCube::StaticClass(); }
-	virtual bool CanFilter() OVERRIDE { return true; }
+	virtual FText GetName() const override { return NSLOCTEXT("AssetTypeActions", "AssetTypeActions_TextureRenderTargetCube", "Cube Render Target"); }
+	virtual UClass* GetSupportedClass() const override { return UTextureRenderTargetCube::StaticClass(); }
+	virtual bool CanFilter() override { return true; }
 };
\ No newline at end of file
diff --git a/Engine/Source/Developer/AssetTools/Private/AssetTypeActions/AssetTypeActions_TouchInterface.h b/Engine/Source/Developer/AssetTools/Private/AssetTypeActions/AssetTypeActions_TouchInterface.h
--- a/Engine/Source/Developer/AssetTools/Private/AssetTypeActions/AssetTypeActions_TouchInterface.h
+++ b/Engine/Source/Developer/AssetTools/Private/AssetTypeActions/AssetTypeActions_TouchInterface.h
@@ -6,8 +6,8 @@ class FAssetTypeActions_TouchInterface : public FAssetTypeActions_Base
 {
 public:
 	// IAssetTypeActions Implementation
-	virtual FText GetName() const OVERRIDE { return NSLOCTEXT("AssetTypeActions", "AssetTypeActions_TouchInterface", "Touch Interface Setup"); }
-	virtual FColor GetTypeColor() const OVERRIDE { return FColor(175, 0, 128); }
-	virtual UClass* GetSupportedClass() const OVERRIDE;
-	virtual uint32 GetCategories() OVERRIDE { return EAssetTypeCategories::Misc; }
+	virtual FText GetName() const override { return NSLOCTEXT("AssetTypeActions", "AssetTypeActions_TouchInterface", "Touch Interface Setup"); }
+	virtual FColor GetTypeColor() const override { return FColor(175, 0, 128); }
+	virtual UClass* GetSupportedClass() const override;
+	virtual uint32 GetCategories() override { return EAssetTypeCategories::Misc; }
 };
diff --git a/Engine/Source/Developer/AssetTools/Private/AssetTypeActions/AssetTypeActions_VectorField.h b/Engine/Source/Developer/AssetTools/Private/AssetTypeActions/AssetTypeActions_VectorField.h
--- a/Engine/Source/Developer/AssetTools/Private/AssetTypeActions/AssetTypeActions_VectorField.h
+++ b/Engine/Source/Developer/AssetTools/Private/AssetTypeActions/AssetTypeActions_VectorField.h
@@ -6,9 +6,9 @@ class FAssetTypeActions_VectorField : public FAssetTypeActions_Base
 {
 public:
 	// IAssetTypeActions Implementation
-	virtual FText GetName() const OVERRIDE { return NSLOCTEXT("AssetTypeActions", "AssetTypeActions_VectorField", "Vector Field"); }
-	virtual FColor GetTypeColor() const OVERRIDE { return FColor(200,128,128); }
-	virtual UClass* GetSupportedClass() const OVERRIDE { return UVectorField::StaticClass(); }
-	virtual bool CanFilter() OVERRIDE { return false; }
-	virtual uint32 GetCategories() OVERRIDE { return EAssetTypeCategories::Misc; }
+	virtual FText GetName() const override { return NSLOCTEXT("AssetTypeActions", "AssetTypeActions_VectorField", "Vector Field"); }
+	virtual FColor GetTypeColor() const override { return FColor(200,128,128); }
+	virtual UClass* GetSupportedClass() const override { return UVectorField::StaticClass(); }
+	virtual bool CanFilter() override { return false; }
+	virtual uint32 GetCategories() override { return EAssetTypeCategories::Misc; }
 };
\ No newline at end of file
diff --git a/Engine/Source/Developer/AssetTools/Private/AssetTypeActions/AssetTypeActions_VectorFieldAnimated.h b/Engine/Source/Developer/AssetTools/Private/AssetTypeActions/AssetTypeActions_VectorFieldAnimated.h
--- a/Engine/Source/Developer/AssetTools/Private/AssetTypeActions/AssetTypeActions_VectorFieldAnimated.h
+++ b/Engine/Source/Developer/AssetTools/Private/AssetTypeActions/AssetTypeActions_VectorFieldAnimated.h
@@ -6,7 +6,7 @@ class FAssetTypeActions_VectorFieldAnimated : public FAssetTypeActions_VectorFie
 {
 public:
 	// IAssetTypeActions Implementation
-	virtual FText GetName() const OVERRIDE { return NSLOCTEXT("AssetTypeActions", "AssetTypeActions_VectorFieldAnimated", "Animated Vector Field"); }
-	virtual UClass* GetSupportedClass() const OVERRIDE { return UVectorFieldAnimated::StaticClass(); }
-	virtual bool CanFilter() OVERRIDE { return true; }
+	virtual FText GetName() const override { return NSLOCTEXT("AssetTypeActions", "AssetTypeActions_VectorFieldAnimated", "Animated Vector Field"); }
+	virtual UClass* GetSupportedClass() const override { return UVectorFieldAnimated::StaticClass(); }
+	virtual bool CanFilter() override { return true; }
 };
\ No newline at end of file
diff --git a/Engine/Source/Developer/AssetTools/Private/AssetTypeActions/AssetTypeActions_VectorFieldStatic.h b/Engine/Source/Developer/AssetTools/Private/AssetTypeActions/AssetTypeActions_VectorFieldStatic.h
--- a/Engine/Source/Developer/AssetTools/Private/AssetTypeActions/AssetTypeActions_VectorFieldStatic.h
+++ b/Engine/Source/Developer/AssetTools/Private/AssetTypeActions/AssetTypeActions_VectorFieldStatic.h
@@ -6,11 +6,11 @@ class FAssetTypeActions_VectorFieldStatic : public FAssetTypeActions_VectorField
 {
 public:
 	// IAssetTypeActions Implementation
-	virtual FText GetName() const OVERRIDE { return NSLOCTEXT("AssetTypeActions", "AssetTypeActions_VectorFieldStatic", "Vector Field"); }
-	virtual UClass* GetSupportedClass() const OVERRIDE { return UVectorFieldStatic::StaticClass(); }
-	virtual bool HasActions ( const TArray<UObject*>& InObjects ) const OVERRIDE { return true; }
-	virtual void GetActions( const TArray<UObject*>& InObjects, FMenuBuilder& MenuBuilder ) OVERRIDE;
-	virtual bool CanFilter() OVERRIDE { return true; }
+	virtual FText GetName() const override { return NSLOCTEXT("AssetTypeActions", "AssetTypeActions_VectorFieldStatic", "Vector Field"); }
+	virtual UClass* GetSupportedClass() const override { return UVectorFieldStatic::StaticClass(); }
+	virtual bool HasActions ( const TArray<UObject*>& InObjects ) const override { return true; }
+	virtual void GetActions( const TArray<UObject*>& InObjects, FMenuBuilder& MenuBuilder ) override;
+	virtual bool CanFilter() override { return true; }
 
 private:
 	/** Handler for when Edit is selected */
diff --git a/Engine/Source/Developer/AssetTools/Private/AssetTypeActions/AssetTypeActions_VertexAnimation.h b/Engine/Source/Developer/AssetTools/Private/AssetTypeActions/AssetTypeActions_VertexAnimation.h
--- a/Engine/Source/Developer/AssetTools/Private/AssetTypeActions/AssetTypeActions_VertexAnimation.h
+++ b/Engine/Source/Developer/AssetTools/Private/AssetTypeActions/AssetTypeActions_VertexAnimation.h
@@ -8,18 +8,18 @@ class FAssetTypeActions_VertexAnimation : public FAssetTypeActions_Base
 {
 public:
 	// IAssetTypeActions Implementation
-	virtual FText GetName() const OVERRIDE { return NSLOCTEXT("AssetTypeActions", "AssetTypeActions_VertexAnimation", "Vertex Animation"); }
-	virtual FColor GetTypeColor() const OVERRIDE { return FColor(87,0,174); }
-	virtual UClass* GetSupportedClass() const OVERRIDE { return UVertexAnimation::StaticClass(); }
-	virtual bool HasActions ( const TArray<UObject*>& InObjects ) const OVERRIDE { return true; }
-	virtual void GetActions( const TArray<UObject*>& InObjects, FMenuBuilder& MenuBuilder ) OVERRIDE;
-	virtual void OpenAssetEditor( const TArray<UObject*>& InObjects, TSharedPtr<class IToolkitHost> EditWithinLevelEditor = TSharedPtr<IToolkitHost>() ) OVERRIDE;
-	virtual uint32 GetCategories() OVERRIDE { return EAssetTypeCategories::Animation; }
+	virtual FText GetName() const override { return NSLOCTEXT("AssetTypeActions", "AssetTypeActions_VertexAnimation", "Vertex Animation"); }
+	virtual FColor GetTypeColor() const override { return FColor(87,0,174); }
+	virtual UClass* GetSupportedClass() const override { return UVertexAnimation::StaticClass(); }
+	virtual bool HasActions ( const TArray<UObject*>& InObjects ) const override { return true; }
+	virtual void GetActions( const TArray<UObject*>& InObjects, FMenuBuilder& MenuBuilder ) override;
+	virtual void OpenAssetEditor( const TArray<UObject*>& InObjects, TSharedPtr<class IToolkitHost> EditWithinLevelEditor = TSharedPtr<IToolkitHost>() ) override;
+	virtual uint32 GetCategories() override { return EAssetTypeCategories::Animation; }
 	
 	// Simply returns false. Vertex Animations are currently are unsupported. Rather than remove
 	// the files in question; returning false from here disables the filer in the content browser.
 	// Remove this when vertex anims are supported again.
-	virtual bool CanFilter() OVERRIDE;
+	virtual bool CanFilter() override;
 
 private:
 	/** Handler for when Edit is selected */
diff --git a/Engine/Source/Developer/AssetTools/Private/AssetTypeActions/AssetTypeActions_World.h b/Engine/Source/Developer/AssetTools/Private/AssetTypeActions/AssetTypeActions_World.h
--- a/Engine/Source/Developer/AssetTools/Private/AssetTypeActions/AssetTypeActions_World.h
+++ b/Engine/Source/Developer/AssetTools/Private/AssetTypeActions/AssetTypeActions_World.h
@@ -6,10 +6,10 @@ class FAssetTypeActions_World : public FAssetTypeActions_Base
 {
 public:
 	// IAssetTypeActions Implementation
-	virtual FText GetName() const OVERRIDE { return NSLOCTEXT("AssetTypeActions", "AssetTypeActions_World", "Level"); }
-	virtual FColor GetTypeColor() const OVERRIDE { return FColor(255, 156, 0); }
-	virtual UClass* GetSupportedClass() const OVERRIDE { return UWorld::StaticClass(); }
-	virtual bool HasActions ( const TArray<UObject*>& InObjects ) const OVERRIDE { return false; }
-	virtual void OpenAssetEditor( const TArray<UObject*>& InObjects, TSharedPtr<class IToolkitHost> EditWithinLevelEditor = TSharedPtr<IToolkitHost>() ) OVERRIDE;
-	virtual uint32 GetCategories() OVERRIDE { return EAssetTypeCategories::Basic; }
+	virtual FText GetName() const override { return NSLOCTEXT("AssetTypeActions", "AssetTypeActions_World", "Level"); }
+	virtual FColor GetTypeColor() const override { return FColor(255, 156, 0); }
+	virtual UClass* GetSupportedClass() const override { return UWorld::StaticClass(); }
+	virtual bool HasActions ( const TArray<UObject*>& InObjects ) const override { return false; }
+	virtual void OpenAssetEditor( const TArray<UObject*>& InObjects, TSharedPtr<class IToolkitHost> EditWithinLevelEditor = TSharedPtr<IToolkitHost>() ) override;
+	virtual uint32 GetCategories() override { return EAssetTypeCategories::Basic; }
 };
\ No newline at end of file
diff --git a/Engine/Source/Developer/AssetTools/Public/AssetTypeActions_Base.h b/Engine/Source/Developer/AssetTools/Public/AssetTypeActions_Base.h
--- a/Engine/Source/Developer/AssetTools/Public/AssetTypeActions_Base.h
+++ b/Engine/Source/Developer/AssetTools/Public/AssetTypeActions_Base.h
@@ -10,22 +10,22 @@ class FAssetTypeActions_Base : public IAssetTypeActions
 {
 public:
 	// Begin IAssetTypeActions implementation
-	virtual bool HasActions( const TArray<UObject*>& InObjects ) const OVERRIDE
+	virtual bool HasActions( const TArray<UObject*>& InObjects ) const override
 	{
 		return false;
 	}
 
-	virtual void GetActions( const TArray<UObject*>& InObjects, FMenuBuilder& MenuBuilder ) OVERRIDE
+	virtual void GetActions( const TArray<UObject*>& InObjects, FMenuBuilder& MenuBuilder ) override
 	{
 
 	}
 
-	virtual void OpenAssetEditor( const TArray<UObject*>& InObjects, TSharedPtr<class IToolkitHost> EditWithinLevelEditor = TSharedPtr<IToolkitHost>() ) OVERRIDE
+	virtual void OpenAssetEditor( const TArray<UObject*>& InObjects, TSharedPtr<class IToolkitHost> EditWithinLevelEditor = TSharedPtr<IToolkitHost>() ) override
 	{
 		FSimpleAssetEditor::CreateEditor(EToolkitMode::Standalone, EditWithinLevelEditor, InObjects);
 	}
 	
-	virtual void AssetsActivated( const TArray<UObject*>& InObjects, EAssetTypeActivationMethod::Type ActivationType ) OVERRIDE
+	virtual void AssetsActivated( const TArray<UObject*>& InObjects, EAssetTypeActivationMethod::Type ActivationType ) override
 	{
 		if ( ActivationType == EAssetTypeActivationMethod::DoubleClicked || ActivationType == EAssetTypeActivationMethod::Opened )
 		{
@@ -40,17 +40,17 @@ public:
 		}
 	}
 
-	virtual bool CanFilter() OVERRIDE
+	virtual bool CanFilter() override
 	{
 		return true;
 	}
 
-	virtual bool ShouldForceWorldCentric() OVERRIDE
+	virtual bool ShouldForceWorldCentric() override
 	{
 		return false;
 	}
 
-	virtual void PerformAssetDiff(UObject* OldAsset, UObject* NewAsset, const struct FRevisionInfo& OldRevision, const struct FRevisionInfo& NewRevision) const OVERRIDE
+	virtual void PerformAssetDiff(UObject* OldAsset, UObject* NewAsset, const struct FRevisionInfo& OldRevision, const struct FRevisionInfo& NewRevision) const override
 	{
 		check(OldAsset != NULL);
 		check(NewAsset != NULL);
@@ -64,7 +64,7 @@ public:
 		AssetToolsModule.Get().CreateDiffProcess(DiffCommand, OldTextFilename, NewTextFilename);
 	}
 
-	virtual class UThumbnailInfo* GetThumbnailInfo(UObject* Asset) const OVERRIDE
+	virtual class UThumbnailInfo* GetThumbnailInfo(UObject* Asset) const override
 	{
 		return NULL;
 	}
diff --git a/Engine/Source/Developer/AudioFormatADPCM/Private/AudioFormatADPCM.cpp b/Engine/Source/Developer/AudioFormatADPCM/Private/AudioFormatADPCM.cpp
--- a/Engine/Source/Developer/AudioFormatADPCM/Private/AudioFormatADPCM.cpp
+++ b/Engine/Source/Developer/AudioFormatADPCM/Private/AudioFormatADPCM.cpp
@@ -432,7 +432,7 @@ public:
 		return false;
 	}
 
-	virtual uint16 GetVersion(FName Format) const OVERRIDE
+	virtual uint16 GetVersion(FName Format) const override
 	{
 		check(Format == NAME_ADPCM);
 		return UE_AUDIO_ADPCM_VER;
diff --git a/Engine/Source/Developer/AudioFormatOgg/Private/AudioFormatOgg.cpp b/Engine/Source/Developer/AudioFormatOgg/Private/AudioFormatOgg.cpp
--- a/Engine/Source/Developer/AudioFormatOgg/Private/AudioFormatOgg.cpp
+++ b/Engine/Source/Developer/AudioFormatOgg/Private/AudioFormatOgg.cpp
@@ -40,7 +40,7 @@ public:
 		return false;
 	}
 
-	virtual uint16 GetVersion(FName Format) const OVERRIDE
+	virtual uint16 GetVersion(FName Format) const override
 	{
 		check(Format == NAME_OGG);
 		return UE_AUDIO_OGG_VER;
diff --git a/Engine/Source/Developer/AudioFormatOpus/Private/AudioFormatOpus.cpp b/Engine/Source/Developer/AudioFormatOpus/Private/AudioFormatOpus.cpp
--- a/Engine/Source/Developer/AudioFormatOpus/Private/AudioFormatOpus.cpp
+++ b/Engine/Source/Developer/AudioFormatOpus/Private/AudioFormatOpus.cpp
@@ -36,7 +36,7 @@ public:
 		return false;
 	}
 
-	virtual uint16 GetVersion(FName Format) const OVERRIDE
+	virtual uint16 GetVersion(FName Format) const override
 	{
 		check(Format == NAME_OPUS);
 		return UE_AUDIO_OPUS_VER;
@@ -234,7 +234,7 @@ public:
 		return CompressedDataStore.Num();
 	}
 
-	virtual bool SplitDataForStreaming(const TArray<uint8>& SrcBuffer, TArray<TArray<uint8>>& OutBuffers) const OVERRIDE
+	virtual bool SplitDataForStreaming(const TArray<uint8>& SrcBuffer, TArray<TArray<uint8>>& OutBuffers) const override
 	{
 		// 16K chunks - don't really have an idea of what's best for loading from disc yet
 		const int32 kMaxChunkSizeBytes = 16384;
diff --git a/Engine/Source/Developer/AutomationController/Private/AutomationControllerManager.h b/Engine/Source/Developer/AutomationController/Private/AutomationControllerManager.h
--- a/Engine/Source/Developer/AutomationController/Private/AutomationControllerManager.h
+++ b/Engine/Source/Developer/AutomationController/Private/AutomationControllerManager.h
@@ -17,176 +17,176 @@ public:
 
 	// Begin IAutomationController Interface
 
-	virtual void RequestAvailableWorkers( const FGuid& InSessionId ) OVERRIDE;
+	virtual void RequestAvailableWorkers( const FGuid& InSessionId ) override;
 
-	virtual void RequestTests() OVERRIDE;
+	virtual void RequestTests() override;
 
-	virtual void RunTests( const bool bIsLocalSession) OVERRIDE;
+	virtual void RunTests( const bool bIsLocalSession) override;
 
-	virtual void StopTests() OVERRIDE;
+	virtual void StopTests() override;
 
-	virtual void Init() OVERRIDE;
+	virtual void Init() override;
 
-	virtual void RequestLoadAsset( const FString& InAssetName ) OVERRIDE;
+	virtual void RequestLoadAsset( const FString& InAssetName ) override;
 
-	virtual void Tick() OVERRIDE;
+	virtual void Tick() override;
 
-	virtual void SetNumPasses(const int32 InNumPasses) OVERRIDE
+	virtual void SetNumPasses(const int32 InNumPasses) override
 	{
 		NumTestPasses = InNumPasses;
 	}
 
-	virtual int32 GetNumPasses() OVERRIDE
+	virtual int32 GetNumPasses() override
 	{
 		return NumTestPasses;
 	}
 
-	virtual bool IsUsingFullSizeScreenshots() const OVERRIDE
+	virtual bool IsUsingFullSizeScreenshots() const override
 	{
 		return bScreenshotsEnabled && bRequestFullScreenScreenshots;
 	}
 
-	virtual void SetUsingFullSizeScreenshots( const bool bNewValue ) OVERRIDE
+	virtual void SetUsingFullSizeScreenshots( const bool bNewValue ) override
 	{
 		bRequestFullScreenScreenshots = bNewValue;
 	}
 
-	virtual bool IsScreenshotAllowed() const OVERRIDE
+	virtual bool IsScreenshotAllowed() const override
 	{
 		return bScreenshotsEnabled;
 	}
 
-	virtual void SetScreenshotsEnabled( const bool bNewValue ) OVERRIDE
+	virtual void SetScreenshotsEnabled( const bool bNewValue ) override
 	{
 		bScreenshotsEnabled = bNewValue;
 	}
 
-	virtual void SetFilter( TSharedPtr< AutomationFilterCollection > InFilter ) OVERRIDE
+	virtual void SetFilter( TSharedPtr< AutomationFilterCollection > InFilter ) override
 	{
 		ReportManager.SetFilter( InFilter );
 	}
 
-	virtual TArray <TSharedPtr <IAutomationReport> >& GetReports() OVERRIDE
+	virtual TArray <TSharedPtr <IAutomationReport> >& GetReports() override
 	{
 		return ReportManager.GetFilteredReports();
 	}
 
-	virtual int32 GetNumDeviceClusters() const OVERRIDE
+	virtual int32 GetNumDeviceClusters() const override
 	{
 		return DeviceClusterManager.GetNumClusters();
 	}
 
-	virtual int32 GetNumDevicesInCluster(const int32 ClusterIndex) const OVERRIDE
+	virtual int32 GetNumDevicesInCluster(const int32 ClusterIndex) const override
 	{
 		return DeviceClusterManager.GetNumDevicesInCluster(ClusterIndex);
 	}
 
-	virtual FString GetClusterGroupName(const int32 ClusterIndex) const OVERRIDE
+	virtual FString GetClusterGroupName(const int32 ClusterIndex) const override
 	{
 		return DeviceClusterManager.GetClusterGroupName(ClusterIndex);
 	}
 
-	virtual FString GetDeviceTypeName(const int32 ClusterIndex) const OVERRIDE
+	virtual FString GetDeviceTypeName(const int32 ClusterIndex) const override
 	{
 		return DeviceClusterManager.GetClusterDeviceType(ClusterIndex);
 	}
 
-	virtual FString GetGameInstanceName(const int32 ClusterIndex, const int32 DeviceIndex) const OVERRIDE
+	virtual FString GetGameInstanceName(const int32 ClusterIndex, const int32 DeviceIndex) const override
 	{
 		return DeviceClusterManager.GetClusterDeviceName(ClusterIndex, DeviceIndex);
 	}
 
-	virtual void SetVisibleTestsEnabled(const bool bEnabled) OVERRIDE
+	virtual void SetVisibleTestsEnabled(const bool bEnabled) override
 	{
 		ReportManager.SetVisibleTestsEnabled (bEnabled);
 	}
 
-	virtual int32 GetEnabledTestsNum() const OVERRIDE
+	virtual int32 GetEnabledTestsNum() const override
 	{
 		return ReportManager.GetEnabledTestsNum();
 	}
 
-	virtual void GetEnabledTestNames(TArray<FString>& OutEnabledTestNames) const OVERRIDE
+	virtual void GetEnabledTestNames(TArray<FString>& OutEnabledTestNames) const override
 	{
 		ReportManager.GetEnabledTestNames(OutEnabledTestNames);
 	}
 
-	virtual void SetEnabledTests(const TArray<FString>& EnabledTests) OVERRIDE
+	virtual void SetEnabledTests(const TArray<FString>& EnabledTests) override
 	{
 		ReportManager.SetEnabledTests(EnabledTests);
 	}
 
-	virtual EAutomationControllerModuleState::Type GetTestState( ) const OVERRIDE
+	virtual EAutomationControllerModuleState::Type GetTestState( ) const override
 	{
 		return AutomationTestState;
 	}
 
-	virtual void SetDeveloperDirectoryIncluded(const bool bInDeveloperDirectoryIncluded) OVERRIDE
+	virtual void SetDeveloperDirectoryIncluded(const bool bInDeveloperDirectoryIncluded) override
 	{
 		bDeveloperDirectoryIncluded = bInDeveloperDirectoryIncluded;
 	}
 
-	virtual bool IsDeveloperDirectoryIncluded(void) const OVERRIDE
+	virtual bool IsDeveloperDirectoryIncluded(void) const override
 	{
 		return bDeveloperDirectoryIncluded;
 	}
 
-	virtual void SetVisualCommandletFilter(const bool bInVisualCommandletFilterOn) OVERRIDE
+	virtual void SetVisualCommandletFilter(const bool bInVisualCommandletFilterOn) override
 	{
 		bVisualCommandletFilterOn = bInVisualCommandletFilterOn;
 	}
 
-	virtual bool IsVisualCommandletFilterOn(void) const OVERRIDE
+	virtual bool IsVisualCommandletFilterOn(void) const override
 	{
 		return bVisualCommandletFilterOn;
 	}
 
-	virtual const bool CheckTestResultsAvailable() const OVERRIDE
+	virtual const bool CheckTestResultsAvailable() const override
 	{
 		return 	bTestResultsAvailable;
 	}
 
-	virtual const bool ReportsHaveErrors() const OVERRIDE
+	virtual const bool ReportsHaveErrors() const override
 	{
 		return bHasErrors;
 	}
 
-	virtual const bool ReportsHaveWarnings() const OVERRIDE
+	virtual const bool ReportsHaveWarnings() const override
 	{
 		return bHasWarning;
 	}
 
-	virtual const bool ReportsHaveLogs() const OVERRIDE
+	virtual const bool ReportsHaveLogs() const override
 	{
 		return bHasLogs;
 	}
 
-	virtual void ClearAutomationReports() OVERRIDE
+	virtual void ClearAutomationReports() override
 	{
 		ReportManager.Empty();
 	}
 
-	virtual const bool ExportReport(uint32 FileExportTypeMask) OVERRIDE;
+	virtual const bool ExportReport(uint32 FileExportTypeMask) override;
 
-	virtual bool IsTestRunnable( IAutomationReportPtr InReport ) const OVERRIDE;
+	virtual bool IsTestRunnable( IAutomationReportPtr InReport ) const override;
 
-	virtual void RemoveCallbacks() OVERRIDE;
+	virtual void RemoveCallbacks() override;
 
-	virtual void Shutdown() OVERRIDE;
+	virtual void Shutdown() override;
 
-	virtual void Startup() OVERRIDE;
+	virtual void Startup() override;
 
-	virtual FOnAutomationControllerManagerShutdown& OnShutdown( ) OVERRIDE
+	virtual FOnAutomationControllerManagerShutdown& OnShutdown( ) override
 	{
 		return ShutdownDelegate;
 	}
 
-	virtual FOnAutomationControllerManagerTestsAvailable& OnTestsAvailable( ) OVERRIDE
+	virtual FOnAutomationControllerManagerTestsAvailable& OnTestsAvailable( ) override
 	{
 		return TestsAvailableDelegate;
 	}
 
-	virtual FOnAutomationControllerTestsRefreshed& OnTestsRefreshed( ) OVERRIDE
+	virtual FOnAutomationControllerTestsRefreshed& OnTestsRefreshed( ) override
 	{
 		return TestsRefreshedDelegate;
 	}
@@ -197,7 +197,7 @@ public:
 
 	void UpdateDeviceGroups();
 
-	virtual void SetTestsCompleteCallback(const FOnAutomationControllerTestsComplete& NewCallback) OVERRIDE
+	virtual void SetTestsCompleteCallback(const FOnAutomationControllerTestsComplete& NewCallback) override
 	{
 		TestsCompleteDelegate = NewCallback;
 	}
