UE4 Refactoring. Changed OVERRIDE and FINAL macros to keywords override and final.

diff --git a/Engine/Source/Developer/AutomationController/Private/AutomationControllerModule.h b/Engine/Source/Developer/AutomationController/Private/AutomationControllerModule.h
--- a/Engine/Source/Developer/AutomationController/Private/AutomationControllerModule.h
+++ b/Engine/Source/Developer/AutomationController/Private/AutomationControllerModule.h
@@ -17,11 +17,11 @@ public:
 
 	// Begin IAutomationControllerModule Interface
 
-	virtual IAutomationControllerManagerRef GetAutomationController( ) OVERRIDE;
+	virtual IAutomationControllerManagerRef GetAutomationController( ) override;
 
-	virtual void Init() OVERRIDE;
+	virtual void Init() override;
 
-	virtual void Tick() OVERRIDE;
+	virtual void Tick() override;
 
 	// End IAutomationControllerModule Interface
 
@@ -30,11 +30,11 @@ public:
 
 	// Begin IModuleInterface interface
 
-	virtual void StartupModule() OVERRIDE;
+	virtual void StartupModule() override;
 
-	virtual void ShutdownModule() OVERRIDE;
+	virtual void ShutdownModule() override;
 
-	virtual bool SupportsDynamicReloading() OVERRIDE;
+	virtual bool SupportsDynamicReloading() override;
 
 	// End IModuleInterface interface
 
diff --git a/Engine/Source/Developer/AutomationController/Private/AutomationReport.h b/Engine/Source/Developer/AutomationController/Private/AutomationReport.h
--- a/Engine/Source/Developer/AutomationController/Private/AutomationReport.h
+++ b/Engine/Source/Developer/AutomationController/Private/AutomationReport.h
@@ -19,15 +19,15 @@ public:
 	// Begin IAutomationReport Interface
 
 	/** Remove all child tests */
-	virtual void Empty() OVERRIDE;
+	virtual void Empty() override;
 
 	/** Get the asset associated with tis report */
-	virtual FString GetAssetName() const OVERRIDE;
+	virtual FString GetAssetName() const override;
 
 	/**
 	 * Returns the complete command line for an automation test including any relevant parameters
 	 */
-	virtual FString GetCommand() const OVERRIDE;
+	virtual FString GetCommand() const override;
 
 	/**
 	 * Returns the name of this level in the test hierarchy for the purposes of grouping.
@@ -35,65 +35,65 @@ public:
 	 * @param bAddDecoration - True is the name should have the number of child tests appended
 	 * @return the name of this level in the hierarchy for use in UI
 	 */
-	virtual const FString& GetDisplayName() const OVERRIDE;
+	virtual const FString& GetDisplayName() const override;
 
 	/**
 	 * Returns the name of this level in the test hierarchy for the purposes of UI.
 	 * (Editor.Maps.LoadAll.parameters) would have 3 internal tree nodes and a variable number of leaf nodes depending on how many maps existed
 	 * @return the name of this level in the hierarchy for use in UI
 	 */
-	virtual FString GetDisplayNameWithDecoration() const OVERRIDE;
+	virtual FString GetDisplayNameWithDecoration() const override;
 
 	/** Recursively gets the number of child nodes */
-	virtual int32 GetTotalNumChildren() const OVERRIDE;
+	virtual int32 GetTotalNumChildren() const override;
 	/** Recursively gets the number of enabled tests */
-	virtual int32 GetEnabledTestsNum() const OVERRIDE;
+	virtual int32 GetEnabledTestsNum() const override;
 	/** Recursively gets the names of enabled tests */
-	virtual void GetEnabledTestNames(TArray<FString>& OutEnabledTestNames, FString CurrentPath) const OVERRIDE;
+	virtual void GetEnabledTestNames(TArray<FString>& OutEnabledTestNames, FString CurrentPath) const override;
 	/** Recursively sets the enabled flag based on if the current test exists in the enabled tests array */
-	virtual void SetEnabledTests(const TArray<FString>& EnabledTests, FString CurrentPath) OVERRIDE;
+	virtual void SetEnabledTests(const TArray<FString>& EnabledTests, FString CurrentPath) override;
 
 	/** Return if this test should be executed */
-	virtual bool IsEnabled() const OVERRIDE;
+	virtual bool IsEnabled() const override;
 	/** Sets whether this test should be executed or not */
-	virtual void SetEnabled(bool bShouldBeEnabled) OVERRIDE;
+	virtual void SetEnabled(bool bShouldBeEnabled) override;
 	
 	/** Sets whether this test is supported on a particular platform */
-	virtual void SetSupport(const int32 ClusterIndex) OVERRIDE;
+	virtual void SetSupport(const int32 ClusterIndex) override;
 	/** Returns if a particular platform is supported */
-	virtual bool IsSupported(const int32 ClusterIndex) const OVERRIDE;
+	virtual bool IsSupported(const int32 ClusterIndex) const override;
 
 	/** Sets the test type for parent reports */
-	virtual void SetTestType( const uint8 TestType ) OVERRIDE;
+	virtual void SetTestType( const uint8 TestType ) override;
 	/** Returns the test type */
-	virtual uint8 GetTestType( ) const OVERRIDE;
+	virtual uint8 GetTestType( ) const override;
 
 	/** Returns true if this is a parent test report */
-	virtual const bool IsParent( ) OVERRIDE;
+	virtual const bool IsParent( ) override;
 	/** Returns true if this is a smoke test report */
-	virtual const bool IsSmokeTest( ) OVERRIDE;
+	virtual const bool IsSmokeTest( ) override;
 	/**
 	 * Filters the visible tests based on search text, execution status, regression test, etc
 	 * @param InFilterDesc - The description of how to decide if this report should be visible
 	 * @param ParentPassedFilter - If the parent passed the filter
 	 * @return whether this report or any of its children passed the filter
 	 */
-	virtual bool SetFilter( TSharedPtr< AutomationFilterCollection > InFilter, const bool ParentPassedFilter = false ) OVERRIDE;
+	virtual bool SetFilter( TSharedPtr< AutomationFilterCollection > InFilter, const bool ParentPassedFilter = false ) override;
 
 	/** Returns the array of child reports that should be visible to the UI based on filtering */
-	virtual TArray<TSharedPtr<IAutomationReport> >& GetFilteredChildren() OVERRIDE;
+	virtual TArray<TSharedPtr<IAutomationReport> >& GetFilteredChildren() override;
 
 	/** Returns the array of child reports */
-	virtual TArray<TSharedPtr<IAutomationReport> >& GetChildReports() OVERRIDE;
+	virtual TArray<TSharedPtr<IAutomationReport> >& GetChildReports() override;
 
 	/** Updates the report when the number of clusters changes */
-	virtual void ClustersUpdated(const int32 NumClusters) OVERRIDE;
+	virtual void ClustersUpdated(const int32 NumClusters) override;
 
 	/** 
 	 * Recursively resets the report to "needs to be run", clears cached warnings and errors
 	 * @param NumTestPasses - The number of test passes so we know how many results to create
 	 */
-	virtual void ResetForExecution(const int32 NumTestPasses) OVERRIDE;
+	virtual void ResetForExecution(const int32 NumTestPasses) override;
 
 	/**
 	 * Sets the results of the test for use by the UI
@@ -101,7 +101,7 @@ public:
 	 * @param PassIndex - Which test pass these results are for
 	 * @param InResults - The new set of results
 	 */
-	virtual void SetResults( const int32 ClusterIndex, const int32 PassIndex, const FAutomationTestResults& InResults ) OVERRIDE;
+	virtual void SetResults( const int32 ClusterIndex, const int32 PassIndex, const FAutomationTestResults& InResults ) override;
 
 	/**
 	 * Returns completion statistics for this branch of the testing hierarchy
@@ -109,7 +109,7 @@ public:
 	 * @param PassIndex - Which test pass to get the status of
 	 * @param OutCompletionState - Collection structure for execution statistics
 	 */
-	virtual void GetCompletionStatus(const int32 ClusterIndex, const int32 PassIndex, FAutomationCompleteState& OutCompletionState) OVERRIDE;
+	virtual void GetCompletionStatus(const int32 ClusterIndex, const int32 PassIndex, FAutomationCompleteState& OutCompletionState) override;
 
 	/**
 	 * Returns the state of the test (not run, in process, success, failure)
@@ -117,7 +117,7 @@ public:
 	 * @param PassIndex - Which test pass to get the state of
 	 * @return the current state of the test
 	 */
-	virtual EAutomationState::Type GetState(const int32 ClusterIndex, const int32 PassIndex) const OVERRIDE;
+	virtual EAutomationState::Type GetState(const int32 ClusterIndex, const int32 PassIndex) const override;
 
 	/**
 	 * Gets a copy of errors and warnings that were found
@@ -127,28 +127,28 @@ public:
 	 *
 	 * @return The collection of results for the given cluster index
 	 */
-	virtual const FAutomationTestResults& GetResults( const int32 ClusterIndex, const int32 PassIndex ) OVERRIDE;
+	virtual const FAutomationTestResults& GetResults( const int32 ClusterIndex, const int32 PassIndex ) override;
 
 	/**
 	 * Gets the number of available test results for a given cluster
 	 * @param ClusterIndex - Index of the platform .
 	 * @return The number of results available for this cluster
 	 */
-	virtual const int32 GetNumResults( const int32 ClusterIndex ) OVERRIDE;
+	virtual const int32 GetNumResults( const int32 ClusterIndex ) override;
 
 	/**
 	 * Finds the current pass by looking at the current state
 	 * @param ClusterIndex - Index of the platform.
 	 * @return The current pass index
 	 */
-	virtual const int32 GetCurrentPassIndex( const int32 ClusterIndex ) OVERRIDE;
+	virtual const int32 GetCurrentPassIndex( const int32 ClusterIndex ) override;
 
 	/**
 	 * Gets the name of the instance that ran the test
 	 * @param ClusterIndex - Index of the platform reporting the results of this test.  See AutomationDeviceClusterManager
 	 * @return the name of the device
 	 */
-	virtual FString GetGameInstanceName( const int32 ClusterIndex ) OVERRIDE;
+	virtual FString GetGameInstanceName( const int32 ClusterIndex ) override;
 
 	/**
 	 * Add a child test to the hierarchy, creating internal tree nodes as needed
@@ -158,7 +158,7 @@ public:
 	 * @param NumPasses - The number of passes we are going to perform.  Used to make sure we have enough results.
 	 * @return - the automation report
 	 */
-	virtual TSharedPtr<IAutomationReport> EnsureReportExists(FAutomationTestInfo& TestInfo, const int32 ClusterIndex, const int32 NumPasses) OVERRIDE;
+	virtual TSharedPtr<IAutomationReport> EnsureReportExists(FAutomationTestInfo& TestInfo, const int32 ClusterIndex, const int32 NumPasses) override;
 
 	/**
 	 * Returns the next test in the hierarchy to run
@@ -169,19 +169,19 @@ public:
 	 *
 	 * @return The next report
 	 */
-	virtual TSharedPtr<IAutomationReport> GetNextReportToExecute(bool& bOutAllTestsComplete, const int32 ClusterIndex, const int32 PassIndex, const int32 NumDevicesInCluster) OVERRIDE;
+	virtual TSharedPtr<IAutomationReport> GetNextReportToExecute(bool& bOutAllTestsComplete, const int32 ClusterIndex, const int32 PassIndex, const int32 NumDevicesInCluster) override;
 
 	/** 
 	 * Have any results got errors
 	 * @return true if they have errors
 	 */
-	virtual const bool HasErrors() OVERRIDE;
+	virtual const bool HasErrors() override;
 
 	/** 
 	 * Have any results got warnings
 	 * @return true if they have warnings
 	 */
-	virtual const bool HasWarnings() OVERRIDE;
+	virtual const bool HasWarnings() override;
 
 	/**
 	 * Gets the min and max time this test took to execute
@@ -189,48 +189,48 @@ public:
 	 * @param OutMaxTime - Maximum execution time for all device types
 	 * @return whether any test has completed successfully
 	 */
-	virtual const bool GetDurationRange(float& OutMinTime, float& OutMaxTime) OVERRIDE;
+	virtual const bool GetDurationRange(float& OutMinTime, float& OutMaxTime) override;
 
 	/**
 	 * Get the number of devices which have been given this test to run
 	 *
 	 * @return - The number of devices who have been given this test to run
 	 */
-	virtual const int32 GetNumDevicesRunningTest() const OVERRIDE;
+	virtual const int32 GetNumDevicesRunningTest() const override;
 
 	/**
 	 * Get the number of participants this test requires.
 	 *
 	 * @Return - The number of devices needed for this test to execute
 	 */
-	virtual const int32 GetNumParticipantsRequired() const OVERRIDE;
+	virtual const int32 GetNumParticipantsRequired() const override;
 	
 	/**
 	 * Set the number of participants this test requires if less than what is already set.
 	 *
 	 * @param NewCount - The number of further devices 
 	 */
-	virtual void SetNumParticipantsRequired( const int32 NewCount ) OVERRIDE;
+	virtual void SetNumParticipantsRequired( const int32 NewCount ) override;
 
 	/**
 	 * Increment the number of network responses
 	 * @return - whether we have received responses from all participants
 	 */
-	virtual bool IncrementNetworkCommandResponses() OVERRIDE;
+	virtual bool IncrementNetworkCommandResponses() override;
 
 	/**
 	 * Resets the number of network responses back to zero
 	 */
-	virtual void ResetNetworkCommandResponses() OVERRIDE;
+	virtual void ResetNetworkCommandResponses() override;
 
 	/**
 	 * Should we expand this node in the UI - A child has passed the filter
 	 * @return true if we should expand the node
 	 */
-	virtual const bool ExpandInUI() const OVERRIDE;
+	virtual const bool ExpandInUI() const override;
 
 	/** Stop the test which is creating this report */
-	virtual void StopRunningTest() OVERRIDE;
+	virtual void StopRunningTest() override;
 
 	/**
 	* Notification on whether we should, or should not, track this reports history
diff --git a/Engine/Source/Developer/AutomationWindow/Private/AutomationWindowModule.cpp b/Engine/Source/Developer/AutomationWindow/Private/AutomationWindowModule.cpp
--- a/Engine/Source/Developer/AutomationWindow/Private/AutomationWindowModule.cpp
+++ b/Engine/Source/Developer/AutomationWindow/Private/AutomationWindowModule.cpp
@@ -17,32 +17,32 @@ public:
 
 	// IAutomationWindowModule interface
 
-	virtual TSharedRef<class SWidget> CreateAutomationWindow( const IAutomationControllerManagerRef& AutomationController, const ISessionManagerRef& SessionManager ) OVERRIDE
+	virtual TSharedRef<class SWidget> CreateAutomationWindow( const IAutomationControllerManagerRef& AutomationController, const ISessionManagerRef& SessionManager ) override
 	{
 		return SNew(SAutomationWindow, AutomationController, SessionManager);
 	}
 
-	virtual TWeakPtr<class SDockTab> GetAutomationWindowTab( ) OVERRIDE
+	virtual TWeakPtr<class SDockTab> GetAutomationWindowTab( ) override
 	{
 		return AutomationWindowTabPtr;
 	}
 
-	virtual FOnAutomationWindowModuleShutdown& OnShutdown( ) OVERRIDE
+	virtual FOnAutomationWindowModuleShutdown& OnShutdown( ) override
 	{
 		return ShutdownDelegate;
 	}
 
-	virtual void SetAutomationWindowTab(TWeakPtr<class SDockTab> AutomationWindowTab) OVERRIDE { AutomationWindowTabPtr = AutomationWindowTab; }
+	virtual void SetAutomationWindowTab(TWeakPtr<class SDockTab> AutomationWindowTab) override { AutomationWindowTabPtr = AutomationWindowTab; }
 
 public:
 
 	// IModuleInterface interface
 
-	virtual void StartupModule( ) OVERRIDE
+	virtual void StartupModule( ) override
 	{
 	}
 
-	virtual void ShutdownModule( ) OVERRIDE
+	virtual void ShutdownModule( ) override
 	{
 		ShutdownDelegate.ExecuteIfBound();	
 	}
diff --git a/Engine/Source/Developer/AutomationWindow/Private/SAutomationTestItem.h b/Engine/Source/Developer/AutomationWindow/Private/SAutomationTestItem.h
--- a/Engine/Source/Developer/AutomationWindow/Private/SAutomationTestItem.h
+++ b/Engine/Source/Developer/AutomationWindow/Private/SAutomationTestItem.h
@@ -40,7 +40,7 @@ public:
 	 * @param ColumnName - The name of the column to generate the widget for.
 	 * @return The widget.
 	 */
-	virtual TSharedRef<SWidget> GenerateWidgetForColumn( const FName& ColumnName ) OVERRIDE;
+	virtual TSharedRef<SWidget> GenerateWidgetForColumn( const FName& ColumnName ) override;
 
 protected:
 
diff --git a/Engine/Source/Developer/AutomationWindow/Private/SAutomationWindow.cpp b/Engine/Source/Developer/AutomationWindow/Private/SAutomationWindow.cpp
--- a/Engine/Source/Developer/AutomationWindow/Private/SAutomationWindow.cpp
+++ b/Engine/Source/Developer/AutomationWindow/Private/SAutomationWindow.cpp
@@ -23,7 +23,7 @@ public:
 	{
 	}
 
-	virtual void RegisterCommands() OVERRIDE
+	virtual void RegisterCommands() override
 	{
 		UI_COMMAND( RefreshTests, "Refresh Tests", "Refresh Tests", EUserInterfaceActionType::Button, FInputGesture() );
 		UI_COMMAND( FindWorkers, "Find Workers", "Find Workers", EUserInterfaceActionType::Button, FInputGesture() );
diff --git a/Engine/Source/Developer/AutomationWindow/Private/SAutomationWindow.h b/Engine/Source/Developer/AutomationWindow/Private/SAutomationWindow.h
--- a/Engine/Source/Developer/AutomationWindow/Private/SAutomationWindow.h
+++ b/Engine/Source/Developer/AutomationWindow/Private/SAutomationWindow.h
@@ -122,10 +122,10 @@ protected:
 	bool IsAnySelectedRowEnabled();
 
 	/** Overridden from SWidget: Called when a key is pressed down - capturing copy */
-	virtual FReply OnKeyDown( const FGeometry& MyGeometry, const FKeyboardEvent& InKeyboardEvent ) OVERRIDE;
+	virtual FReply OnKeyDown( const FGeometry& MyGeometry, const FKeyboardEvent& InKeyboardEvent ) override;
 
 	/** Overridden from SWidget: Called after a key is released */
-	virtual FReply OnKeyUp( const FGeometry& InGeometry, const FKeyboardEvent& InKeyboardEvent ) OVERRIDE;
+	virtual FReply OnKeyUp( const FGeometry& InGeometry, const FKeyboardEvent& InKeyboardEvent ) override;
 
 private:
 
diff --git a/Engine/Source/Developer/CollectionManager/Private/CollectionManager.h b/Engine/Source/Developer/CollectionManager/Private/CollectionManager.h
--- a/Engine/Source/Developer/CollectionManager/Private/CollectionManager.h
+++ b/Engine/Source/Developer/CollectionManager/Private/CollectionManager.h
@@ -9,43 +9,43 @@ public:
 	virtual ~FCollectionManager();
 
 	// ICollectionManager implementation
-	virtual void GetCollectionNames(ECollectionShareType::Type ShareType, TArray<FName>& CollectionNames) const OVERRIDE;
-	virtual bool CollectionExists(FName CollectionName, ECollectionShareType::Type ShareType) const OVERRIDE;
-	virtual bool GetAssetsInCollection(FName CollectionName, ECollectionShareType::Type ShareType, TArray<FName>& AssetPaths) const OVERRIDE;
-	virtual bool GetObjectsInCollection(FName CollectionName, ECollectionShareType::Type ShareType, TArray<FName>& ObjectPaths) const OVERRIDE;
-	virtual bool GetClassesInCollection(FName CollectionName, ECollectionShareType::Type ShareType, TArray<FName>& ClassPaths) const OVERRIDE;
-	virtual void GetCollectionsContainingAsset(FName ObjectPath, ECollectionShareType::Type ShareType, TArray<FName>& OutCollectionNames) const OVERRIDE;
-	virtual void CreateUniqueCollectionName(const FName& BaseName, ECollectionShareType::Type ShareType, FName& OutCollectionName) const OVERRIDE;
-	virtual bool CreateCollection(FName CollectionName, ECollectionShareType::Type ShareType) OVERRIDE;
-	virtual bool RenameCollection(FName CurrentCollectionName, ECollectionShareType::Type CurrentShareType, FName NewCollectionName, ECollectionShareType::Type NewShareType) OVERRIDE;
-	virtual bool DestroyCollection(FName CollectionName, ECollectionShareType::Type ShareType) OVERRIDE;
-	virtual bool AddToCollection(FName CollectionName, ECollectionShareType::Type ShareType, FName ObjectPath) OVERRIDE;
-	virtual bool AddToCollection(FName CollectionName, ECollectionShareType::Type ShareType, const TArray<FName>& ObjectPaths, int32* OutNumAdded = NULL) OVERRIDE;
-	virtual bool RemoveFromCollection(FName CollectionName, ECollectionShareType::Type ShareType, FName ObjectPath) OVERRIDE;
-	virtual bool RemoveFromCollection(FName CollectionName, ECollectionShareType::Type ShareType, const TArray<FName>& ObjectPaths, int32* OutNumRemoved = NULL) OVERRIDE;
-	virtual bool EmptyCollection(FName CollectionName, ECollectionShareType::Type ShareType) OVERRIDE;
-	virtual bool IsCollectionEmpty(FName CollectionName, ECollectionShareType::Type ShareType) OVERRIDE;
-	virtual FText GetLastError() const OVERRIDE { return LastError; }
+	virtual void GetCollectionNames(ECollectionShareType::Type ShareType, TArray<FName>& CollectionNames) const override;
+	virtual bool CollectionExists(FName CollectionName, ECollectionShareType::Type ShareType) const override;
+	virtual bool GetAssetsInCollection(FName CollectionName, ECollectionShareType::Type ShareType, TArray<FName>& AssetPaths) const override;
+	virtual bool GetObjectsInCollection(FName CollectionName, ECollectionShareType::Type ShareType, TArray<FName>& ObjectPaths) const override;
+	virtual bool GetClassesInCollection(FName CollectionName, ECollectionShareType::Type ShareType, TArray<FName>& ClassPaths) const override;
+	virtual void GetCollectionsContainingAsset(FName ObjectPath, ECollectionShareType::Type ShareType, TArray<FName>& OutCollectionNames) const override;
+	virtual void CreateUniqueCollectionName(const FName& BaseName, ECollectionShareType::Type ShareType, FName& OutCollectionName) const override;
+	virtual bool CreateCollection(FName CollectionName, ECollectionShareType::Type ShareType) override;
+	virtual bool RenameCollection(FName CurrentCollectionName, ECollectionShareType::Type CurrentShareType, FName NewCollectionName, ECollectionShareType::Type NewShareType) override;
+	virtual bool DestroyCollection(FName CollectionName, ECollectionShareType::Type ShareType) override;
+	virtual bool AddToCollection(FName CollectionName, ECollectionShareType::Type ShareType, FName ObjectPath) override;
+	virtual bool AddToCollection(FName CollectionName, ECollectionShareType::Type ShareType, const TArray<FName>& ObjectPaths, int32* OutNumAdded = NULL) override;
+	virtual bool RemoveFromCollection(FName CollectionName, ECollectionShareType::Type ShareType, FName ObjectPath) override;
+	virtual bool RemoveFromCollection(FName CollectionName, ECollectionShareType::Type ShareType, const TArray<FName>& ObjectPaths, int32* OutNumRemoved = NULL) override;
+	virtual bool EmptyCollection(FName CollectionName, ECollectionShareType::Type ShareType) override;
+	virtual bool IsCollectionEmpty(FName CollectionName, ECollectionShareType::Type ShareType) override;
+	virtual FText GetLastError() const override { return LastError; }
 
 	/** Event for when collections are created */
 	DECLARE_DERIVED_EVENT( FCollectionManager, ICollectionManager::FCollectionCreatedEvent, FCollectionCreatedEvent );
-	virtual FCollectionCreatedEvent& OnCollectionCreated() OVERRIDE { return CollectionCreatedEvent; }
+	virtual FCollectionCreatedEvent& OnCollectionCreated() override { return CollectionCreatedEvent; }
 
 	/** Event for when collections are destroyed */
 	DECLARE_DERIVED_EVENT( FCollectionManager, ICollectionManager::FCollectionDestroyedEvent, FCollectionDestroyedEvent );
-	virtual FCollectionDestroyedEvent& OnCollectionDestroyed() OVERRIDE { return CollectionDestroyedEvent; }
+	virtual FCollectionDestroyedEvent& OnCollectionDestroyed() override { return CollectionDestroyedEvent; }
 
 	/** Event for when assets are added to a collection */
 	DECLARE_DERIVED_EVENT( FCollectionManager, ICollectionManager::FAssetsAddedEvent, FAssetsAddedEvent ); 
-	virtual FAssetsAddedEvent& OnAssetsAdded() OVERRIDE { return AssetsAddedEvent; }
+	virtual FAssetsAddedEvent& OnAssetsAdded() override { return AssetsAddedEvent; }
 
 	/** Event for when assets are removed to a collection */
 	DECLARE_DERIVED_EVENT( FCollectionManager, ICollectionManager::FAssetsRemovedEvent, FAssetsRemovedEvent );
-	virtual FAssetsRemovedEvent& OnAssetsRemoved() OVERRIDE { return AssetsRemovedEvent; }
+	virtual FAssetsRemovedEvent& OnAssetsRemoved() override { return AssetsRemovedEvent; }
 
 	/** Event for when collections are renamed */
 	DECLARE_DERIVED_EVENT( FCollectionManager, ICollectionManager::FCollectionRenamedEvent, FCollectionRenamedEvent );
-	virtual FCollectionRenamedEvent& OnCollectionRenamed() OVERRIDE { return CollectionRenamedEvent; }
+	virtual FCollectionRenamedEvent& OnCollectionRenamed() override { return CollectionRenamedEvent; }
 
 private:
 	/** Loads all collection files from disk */
diff --git a/Engine/Source/Developer/CollisionAnalyzer/Private/CollisionAnalyzer.h b/Engine/Source/Developer/CollisionAnalyzer/Private/CollisionAnalyzer.h
--- a/Engine/Source/Developer/CollisionAnalyzer/Private/CollisionAnalyzer.h
+++ b/Engine/Source/Developer/CollisionAnalyzer/Private/CollisionAnalyzer.h
@@ -39,13 +39,13 @@ public:
 		const FCollisionObjectQueryParams&	ObjectParams,
 		const TArray<FHitResult>& Results, 
 		const TArray<FHitResult>& TouchAllResults,
-		double CPUTime) OVERRIDE;
+		double CPUTime) override;
 
 	/** Returns a new Collision Analyzer widget. */
-	virtual TSharedPtr<SWidget> SummonUI() OVERRIDE;
+	virtual TSharedPtr<SWidget> SummonUI() override;
 
-	virtual void TickAnalyzer(UWorld* InWorld) OVERRIDE;
-	virtual bool IsRecording() OVERRIDE;
+	virtual void TickAnalyzer(UWorld* InWorld) override;
+	virtual bool IsRecording() override;
 	// End ICollisionAnalyzer interface
 
 	/** Change the current recording state */
diff --git a/Engine/Source/Developer/CollisionAnalyzer/Private/SCAQueryDetails.cpp b/Engine/Source/Developer/CollisionAnalyzer/Private/SCAQueryDetails.cpp
--- a/Engine/Source/Developer/CollisionAnalyzer/Private/SCAQueryDetails.cpp
+++ b/Engine/Source/Developer/CollisionAnalyzer/Private/SCAQueryDetails.cpp
@@ -53,7 +53,7 @@ public:
 	}
 
 	BEGIN_SLATE_FUNCTION_BUILD_OPTIMIZATION
-	virtual TSharedRef<SWidget> GenerateWidgetForColumn(const FName& ColumnName) OVERRIDE
+	virtual TSharedRef<SWidget> GenerateWidgetForColumn(const FName& ColumnName) override
 	{
 		// Get info to apply to all columns (color and tooltip)
 		FSlateColor ResultColor = FSlateColor::UseForeground();
diff --git a/Engine/Source/Developer/CollisionAnalyzer/Private/SCAQueryTableRow.h b/Engine/Source/Developer/CollisionAnalyzer/Private/SCAQueryTableRow.h
--- a/Engine/Source/Developer/CollisionAnalyzer/Private/SCAQueryTableRow.h
+++ b/Engine/Source/Developer/CollisionAnalyzer/Private/SCAQueryTableRow.h
@@ -18,7 +18,7 @@ public:
 
 	void Construct(const FArguments& InArgs, const TSharedRef<STableViewBase>& InOwnerTableView);
 
-	virtual TSharedRef<SWidget> GenerateWidgetForColumn(const FName& ColumnName) OVERRIDE;
+	virtual TSharedRef<SWidget> GenerateWidgetForColumn(const FName& ColumnName) override;
 
 	// delegate
 	FString GetTotalTimeString() const;
diff --git a/Engine/Source/Developer/CrashDebugHelper/Private/Linux/CrashDebugHelperLinux.h b/Engine/Source/Developer/CrashDebugHelper/Private/Linux/CrashDebugHelperLinux.h
--- a/Engine/Source/Developer/CrashDebugHelper/Private/Linux/CrashDebugHelperLinux.h
+++ b/Engine/Source/Developer/CrashDebugHelper/Private/Linux/CrashDebugHelperLinux.h
@@ -16,9 +16,9 @@ public:
 	 *
 	 *	@return	bool				true if successful, false if not
 	 */
-	virtual bool ParseCrashDump(const FString& InCrashDumpName, FCrashDebugInfo& OutCrashDebugInfo) OVERRIDE;
+	virtual bool ParseCrashDump(const FString& InCrashDumpName, FCrashDebugInfo& OutCrashDebugInfo) override;
 
-	virtual bool SyncAndDebugCrashDump(const FString& InCrashDumpName) OVERRIDE;
+	virtual bool SyncAndDebugCrashDump(const FString& InCrashDumpName) override;
 
 	/**
 	 *	Parse the given crash dump, and generate a report. 
@@ -27,7 +27,7 @@ public:
 	 *
 	 *	@return	bool				true if successful, false if not
 	 */
-	virtual bool CreateMinidumpDiagnosticReport( const FString& InCrashDumpName ) OVERRIDE;
+	virtual bool CreateMinidumpDiagnosticReport( const FString& InCrashDumpName ) override;
 };
 
 typedef FCrashDebugHelperLinux FCrashDebugHelper;
diff --git a/Engine/Source/Developer/CrashDebugHelper/Private/Mac/CrashDebugHelperMac.h b/Engine/Source/Developer/CrashDebugHelper/Private/Mac/CrashDebugHelperMac.h
--- a/Engine/Source/Developer/CrashDebugHelper/Private/Mac/CrashDebugHelperMac.h
+++ b/Engine/Source/Developer/CrashDebugHelper/Private/Mac/CrashDebugHelperMac.h
@@ -16,9 +16,9 @@ public:
 	 *
 	 *	@return	bool				true if successful, false if not
 	 */
-	virtual bool ParseCrashDump(const FString& InCrashDumpName, FCrashDebugInfo& OutCrashDebugInfo) OVERRIDE;
+	virtual bool ParseCrashDump(const FString& InCrashDumpName, FCrashDebugInfo& OutCrashDebugInfo) override;
 
-	virtual bool SyncAndDebugCrashDump(const FString& InCrashDumpName) OVERRIDE;
+	virtual bool SyncAndDebugCrashDump(const FString& InCrashDumpName) override;
 
 	/**
 	 *	Parse the given crash dump, and generate a report. 
@@ -27,7 +27,7 @@ public:
 	 *
 	 *	@return	bool				true if successful, false if not
 	 */
-	virtual bool CreateMinidumpDiagnosticReport( const FString& InCrashDumpName ) OVERRIDE;
+	virtual bool CreateMinidumpDiagnosticReport( const FString& InCrashDumpName ) override;
 };
 
 typedef FCrashDebugHelperMac FCrashDebugHelper;
diff --git a/Engine/Source/Developer/DerivedDataCache/Private/DerivedDataBackendAsyncPutWrapper.h b/Engine/Source/Developer/DerivedDataCache/Private/DerivedDataBackendAsyncPutWrapper.h
--- a/Engine/Source/Developer/DerivedDataCache/Private/DerivedDataBackendAsyncPutWrapper.h
+++ b/Engine/Source/Developer/DerivedDataCache/Private/DerivedDataBackendAsyncPutWrapper.h
@@ -170,7 +170,7 @@ public:
 	 * @param	InData		Buffer containing the data to cache, can be destroyed after the call returns, immediately
 	 * @param	bPutEvenIfExists	If true, then do not attempt skip the put even if CachedDataProbablyExists returns true
 	 */
-	virtual void PutCachedData(const TCHAR* CacheKey, TArray<uint8>& InData, bool bPutEvenIfExists) OVERRIDE
+	virtual void PutCachedData(const TCHAR* CacheKey, TArray<uint8>& InData, bool bPutEvenIfExists) override
 	{
 		if (!InnerBackend->IsWritable())
 		{
@@ -193,7 +193,7 @@ public:
 		(new FAutoDeleteAsyncTask<FCachePutAsyncWorker>(CacheKey, &InData, InnerBackend, bPutEvenIfExists, InflightCache.GetOwnedPointer(), &FilesInFlight))->StartBackgroundTask();
 	}
 
-	virtual void RemoveCachedData(const TCHAR* CacheKey, bool bTransient) OVERRIDE
+	virtual void RemoveCachedData(const TCHAR* CacheKey, bool bTransient) override
 	{
 		if (!InnerBackend->IsWritable())
 		{
diff --git a/Engine/Source/Developer/DerivedDataCache/Private/DerivedDataBackendCorruptionWrapper.h b/Engine/Source/Developer/DerivedDataCache/Private/DerivedDataBackendCorruptionWrapper.h
--- a/Engine/Source/Developer/DerivedDataCache/Private/DerivedDataBackendCorruptionWrapper.h
+++ b/Engine/Source/Developer/DerivedDataCache/Private/DerivedDataBackendCorruptionWrapper.h
@@ -144,7 +144,7 @@ public:
 	 * @param	InData		Buffer containing the data to cache, can be destroyed after the call returns, immediately
 	 * @param	bPutEvenIfExists	If true, then do not attempt skip the put even if CachedDataProbablyExists returns true
 	 */
-	virtual void PutCachedData(const TCHAR* CacheKey, TArray<uint8>& InData, bool bPutEvenIfExists) OVERRIDE
+	virtual void PutCachedData(const TCHAR* CacheKey, TArray<uint8>& InData, bool bPutEvenIfExists) override
 	{
 		if (!InnerBackend->IsWritable())
 		{
@@ -162,7 +162,7 @@ public:
 		InnerBackend->PutCachedData(CacheKey, Data, bPutEvenIfExists);
 	}
 	
-	virtual void RemoveCachedData(const TCHAR* CacheKey, bool bTransient) OVERRIDE
+	virtual void RemoveCachedData(const TCHAR* CacheKey, bool bTransient) override
 	{
 		if (!InnerBackend->IsWritable())
 		{
diff --git a/Engine/Source/Developer/DerivedDataCache/Private/DerivedDataBackends.cpp b/Engine/Source/Developer/DerivedDataCache/Private/DerivedDataBackends.cpp
--- a/Engine/Source/Developer/DerivedDataCache/Private/DerivedDataBackends.cpp
+++ b/Engine/Source/Developer/DerivedDataCache/Private/DerivedDataBackends.cpp
@@ -627,13 +627,13 @@ public:
 		DestroyCreatedBackends();
 	}
 
-	FDerivedDataBackendInterface& GetRoot() OVERRIDE
+	FDerivedDataBackendInterface& GetRoot() override
 	{
 		check(RootCache);
 		return *RootCache;
 	}
 
-	virtual void NotifyBootComplete() OVERRIDE
+	virtual void NotifyBootComplete() override
 	{
 		check(RootCache);
 		if (BootCache)
@@ -646,7 +646,7 @@ public:
 		}
 	}
 
-	virtual void WaitForQuiescence(bool bShutdown) OVERRIDE
+	virtual void WaitForQuiescence(bool bShutdown) override
 	{
 		double StartTime = FPlatformTime::Seconds();
 		double LastPrint = StartTime;
@@ -700,13 +700,13 @@ public:
 		}
 	}
 
-	virtual void AddToAsyncCompletionCounter(int32 Addend) OVERRIDE
+	virtual void AddToAsyncCompletionCounter(int32 Addend) override
 	{
 		AsyncCompletionCounter.Add(Addend);
 		check(AsyncCompletionCounter.GetValue() >= 0);
 	}
 
-	virtual void GetDirectories(TArray<FString>& OutResults) OVERRIDE
+	virtual void GetDirectories(TArray<FString>& OutResults) override
 	{
 		OutResults = Directories;
 	}
@@ -717,7 +717,7 @@ public:
 		return SingletonInstance;
 	}
 
-	virtual FDerivedDataBackendInterface* MountPakFile(const TCHAR* PakFilename) OVERRIDE
+	virtual FDerivedDataBackendInterface* MountPakFile(const TCHAR* PakFilename) override
 	{
 		// Assumptions: there's at least one read-only pak backend in the hierarchy
 		// and its parent is a hierarchical backend.
@@ -738,7 +738,7 @@ public:
 		return ReadPak;
 	}
 
-	virtual bool UnmountPakFile(const TCHAR* PakFilename) OVERRIDE
+	virtual bool UnmountPakFile(const TCHAR* PakFilename) override
 	{
 		for (int PakIndex = 0; PakIndex < ReadPakCache.Num(); ++PakIndex)
 		{
diff --git a/Engine/Source/Developer/DerivedDataCache/Private/DerivedDataCache.cpp b/Engine/Source/Developer/DerivedDataCache/Private/DerivedDataCache.cpp
--- a/Engine/Source/Developer/DerivedDataCache/Private/DerivedDataCache.cpp
+++ b/Engine/Source/Developer/DerivedDataCache/Private/DerivedDataCache.cpp
@@ -173,7 +173,7 @@ public:
 		return Handle;
 	}
 
-	virtual bool PollAsynchronousCompletion(uint32 Handle) OVERRIDE
+	virtual bool PollAsynchronousCompletion(uint32 Handle) override
 	{
 		FAsyncTask<FBuildAsyncWorker>* AsyncTask = NULL;
 		{
@@ -184,7 +184,7 @@ public:
 		return AsyncTask->IsDone();
 	}
 
-	virtual void WaitAsynchronousCompletion(uint32 Handle) OVERRIDE
+	virtual void WaitAsynchronousCompletion(uint32 Handle) override
 	{
 		STAT(double ThisTime = 0);
 		{
@@ -200,7 +200,7 @@ public:
 		INC_FLOAT_STAT_BY(STAT_DDC_ASyncWaitTime,(float)ThisTime);
 	}
 
-	virtual bool GetAsynchronousResults(uint32 Handle, TArray<uint8>& OutData) OVERRIDE
+	virtual bool GetAsynchronousResults(uint32 Handle, TArray<uint8>& OutData) override
 	{
 		FAsyncTask<FBuildAsyncWorker>* AsyncTask = NULL;
 		{
@@ -219,17 +219,17 @@ public:
 		return true;
 	}
 
-	virtual IDerivedDataRollup* StartRollup() OVERRIDE
+	virtual IDerivedDataRollup* StartRollup() override
 	{
 		return NULL;
 	}
 
-	virtual void EndRollup(IDerivedDataRollup*& Rollup) OVERRIDE
+	virtual void EndRollup(IDerivedDataRollup*& Rollup) override
 	{
 		check(!Rollup); // this needs to be handled by someone else, if rollups are disabled, then it should be NULL
 	}
 
-	virtual bool GetSynchronous(const TCHAR* CacheKey, TArray<uint8>& OutData) OVERRIDE
+	virtual bool GetSynchronous(const TCHAR* CacheKey, TArray<uint8>& OutData) override
 	{
 		UE_LOG(LogDerivedDataCache, Verbose, TEXT("GetSynchronous %s"), CacheKey);
 		FAsyncTask<FBuildAsyncWorker> PendingTask((FDerivedDataPluginInterface*)NULL, CacheKey, true);
@@ -239,7 +239,7 @@ public:
 		return PendingTask.GetTask().bSuccess;
 	}
 
-	virtual uint32 GetAsynchronous(const TCHAR* CacheKey, IDerivedDataRollup* Rollup = NULL) OVERRIDE
+	virtual uint32 GetAsynchronous(const TCHAR* CacheKey, IDerivedDataRollup* Rollup = NULL) override
 	{
 		check(!Rollup); // this needs to be handled by someone else, if rollups are disabled, then it should be NULL
 		FScopeLock ScopeLock(&SynchronizationObject);
@@ -270,7 +270,7 @@ public:
 		AsyncTask->StartBackgroundTask();
 	}
 
-	virtual void Put(const TCHAR* CacheKey, TArray<uint8>& Data, bool bPutEvenIfExists = false) OVERRIDE
+	virtual void Put(const TCHAR* CacheKey, TArray<uint8>& Data, bool bPutEvenIfExists = false) override
 	{
 		STAT(double ThisTime = 0);
 		{
@@ -281,32 +281,32 @@ public:
 		INC_DWORD_STAT(STAT_DDC_NumPuts);
 	}
 
-	virtual void MarkTransient(const TCHAR* CacheKey) OVERRIDE
+	virtual void MarkTransient(const TCHAR* CacheKey) override
 	{
 		FDerivedDataBackend::Get().GetRoot().RemoveCachedData(CacheKey, /*bTransient=*/ true);
 	}
 
-	virtual bool CachedDataProbablyExists(const TCHAR* CacheKey) OVERRIDE
+	virtual bool CachedDataProbablyExists(const TCHAR* CacheKey) override
 	{
 		return FDerivedDataBackend::Get().GetRoot().CachedDataProbablyExists(CacheKey);
 	}
 
-	void NotifyBootComplete() OVERRIDE
+	void NotifyBootComplete() override
 	{
 		FDerivedDataBackend::Get().NotifyBootComplete();
 	}
 
-	void AddToAsyncCompletionCounter(int32 Addend) OVERRIDE
+	void AddToAsyncCompletionCounter(int32 Addend) override
 	{
 		FDerivedDataBackend::Get().AddToAsyncCompletionCounter(Addend);
 	}
 
-	void WaitForQuiescence(bool bShutdown) OVERRIDE
+	void WaitForQuiescence(bool bShutdown) override
 	{
 		FDerivedDataBackend::Get().WaitForQuiescence(bShutdown);
 	}
 
-	void GetDirectories(TArray<FString>& OutResults) OVERRIDE
+	void GetDirectories(TArray<FString>& OutResults) override
 	{
 		FDerivedDataBackend::Get().GetDirectories(OutResults);
 	}
@@ -742,7 +742,7 @@ public:
 	{
 	}
 
-	virtual void PrintLeaks() OVERRIDE
+	virtual void PrintLeaks() override
 	{
 		FScopeLock ScopeLock(&SynchronizationObject);
 		UE_LOG(LogDerivedDataCache, Log, TEXT("Shutdown"));
@@ -759,7 +759,7 @@ public:
 		//check(!PendingRollups.Num()); // leaked rollups
 	}
 
-	virtual bool PollAsynchronousCompletion(uint32 Handle) OVERRIDE
+	virtual bool PollAsynchronousCompletion(uint32 Handle) override
 	{
 		FScopeLock ScopeLock(&SynchronizationObject);
 		for (TSet<FDerivedDataRollup*>::TIterator Iter(PendingRollups); Iter; ++Iter)
@@ -772,7 +772,7 @@ public:
 		return Super::PollAsynchronousCompletion(Handle);
 	}
 
-	virtual void WaitAsynchronousCompletion(uint32 Handle) OVERRIDE
+	virtual void WaitAsynchronousCompletion(uint32 Handle) override
 	{
 		STAT(double ThisTime = 0);
 		{
@@ -791,7 +791,7 @@ public:
 		Super::WaitAsynchronousCompletion(Handle);
 	}
 
-	virtual bool GetAsynchronousResults(uint32 Handle, TArray<uint8>& OutData) OVERRIDE
+	virtual bool GetAsynchronousResults(uint32 Handle, TArray<uint8>& OutData) override
 	{
 		FScopeLock ScopeLock(&SynchronizationObject);
 		for (TSet<FDerivedDataRollup*>::TIterator Iter(PendingRollups); Iter; ++Iter)
@@ -810,7 +810,7 @@ public:
 		return Super::GetAsynchronousResults(Handle, OutData);
 	}
 
-	virtual IDerivedDataRollup* StartRollup() OVERRIDE
+	virtual IDerivedDataRollup* StartRollup() override
 	{
 		FScopeLock ScopeLock(&SynchronizationObject);
 		FDerivedDataRollup* Result = new FDerivedDataRollup();
@@ -818,7 +818,7 @@ public:
 		return Result;
 	}
 
-	virtual void EndRollup(IDerivedDataRollup*& InRollup) OVERRIDE
+	virtual void EndRollup(IDerivedDataRollup*& InRollup) override
 	{
 		FDerivedDataRollup* Rollup = static_cast<FDerivedDataRollup*>(InRollup);
 		if (Rollup)
@@ -834,7 +834,7 @@ public:
 		}
 	}
 
-	virtual uint32 GetAsynchronous(const TCHAR* CacheKey, IDerivedDataRollup* Rollup = NULL) OVERRIDE
+	virtual uint32 GetAsynchronous(const TCHAR* CacheKey, IDerivedDataRollup* Rollup = NULL) override
 	{
 		if (Rollup)
 		{
@@ -847,7 +847,7 @@ public:
 		return Super::GetAsynchronous(CacheKey, NULL);
 	}
 
-	virtual void Put(const TCHAR* CacheKey, TArray<uint8>& Data, bool bPutEvenIfExists = false) OVERRIDE
+	virtual void Put(const TCHAR* CacheKey, TArray<uint8>& Data, bool bPutEvenIfExists = false) override
 	{
 		FScopeLock ScopeLock(&SynchronizationObject);
 		for (TSet<FDerivedDataRollup*>::TIterator Iter(PendingRollups); Iter; ++Iter)
@@ -903,12 +903,12 @@ FDerivedDataCache& InternalSingleton()
 class FDerivedDataCacheModule : public IDerivedDataCacheModule
 {
 public:
-	virtual FDerivedDataCacheInterface& GetDDC() OVERRIDE
+	virtual FDerivedDataCacheInterface& GetDDC() override
 	{
 		return InternalSingleton();
 	}
 
-	virtual void ShutdownModule() OVERRIDE
+	virtual void ShutdownModule() override
 	{
 		FDDCCleanup::Shutdown();
 
diff --git a/Engine/Source/Developer/DerivedDataCache/Private/FileSystemDerivedDataBackend.cpp b/Engine/Source/Developer/DerivedDataCache/Private/FileSystemDerivedDataBackend.cpp
--- a/Engine/Source/Developer/DerivedDataCache/Private/FileSystemDerivedDataBackend.cpp
+++ b/Engine/Source/Developer/DerivedDataCache/Private/FileSystemDerivedDataBackend.cpp
@@ -172,7 +172,7 @@ public:
 	 * @param	OutData		Buffer containing the data to cache, can be destroyed after the call returns, immediately
 	 * @param	bPutEvenIfExists	If true, then do not attempt skip the put even if CachedDataProbablyExists returns true
 	 */
-	virtual void PutCachedData(const TCHAR* CacheKey, TArray<uint8>& Data, bool bPutEvenIfExists) OVERRIDE
+	virtual void PutCachedData(const TCHAR* CacheKey, TArray<uint8>& Data, bool bPutEvenIfExists) override
 	{
 		check(!bFailed);
 		if (!bReadOnly)
@@ -228,7 +228,7 @@ public:
 		}
 	}
 
-	void RemoveCachedData(const TCHAR* CacheKey, bool bTransient) OVERRIDE
+	void RemoveCachedData(const TCHAR* CacheKey, bool bTransient) override
 	{
 		check(!bFailed);
 		if (!bReadOnly && (!bTransient || bPurgeTransient))
diff --git a/Engine/Source/Developer/DesktopPlatform/Private/DesktopPlatformBase.h b/Engine/Source/Developer/DesktopPlatform/Private/DesktopPlatformBase.h
--- a/Engine/Source/Developer/DesktopPlatform/Private/DesktopPlatformBase.h
+++ b/Engine/Source/Developer/DesktopPlatform/Private/DesktopPlatformBase.h
@@ -12,33 +12,33 @@ public:
 	FDesktopPlatformBase();
 
 	// IDesktopPlatform Implementation
-	virtual FString GetEngineDescription(const FString& Identifier) OVERRIDE;
-	virtual FString GetCurrentEngineIdentifier() OVERRIDE;
+	virtual FString GetEngineDescription(const FString& Identifier) override;
+	virtual FString GetCurrentEngineIdentifier() override;
 
-	virtual void EnumerateLauncherEngineInstallations(TMap<FString, FString> &OutInstallations) OVERRIDE;
-	virtual void EnumerateLauncherSampleInstallations(TArray<FString> &OutInstallations) OVERRIDE;
-	virtual void EnumerateLauncherSampleProjects(TArray<FString> &OutFileNames) OVERRIDE;
-	virtual bool GetEngineRootDirFromIdentifier(const FString &Identifier, FString &OutRootDir) OVERRIDE;
-	virtual bool GetEngineIdentifierFromRootDir(const FString &RootDir, FString &OutIdentifier) OVERRIDE;
+	virtual void EnumerateLauncherEngineInstallations(TMap<FString, FString> &OutInstallations) override;
+	virtual void EnumerateLauncherSampleInstallations(TArray<FString> &OutInstallations) override;
+	virtual void EnumerateLauncherSampleProjects(TArray<FString> &OutFileNames) override;
+	virtual bool GetEngineRootDirFromIdentifier(const FString &Identifier, FString &OutRootDir) override;
+	virtual bool GetEngineIdentifierFromRootDir(const FString &RootDir, FString &OutIdentifier) override;
 
-	virtual bool GetDefaultEngineIdentifier(FString &OutIdentifier) OVERRIDE;
-	virtual bool GetDefaultEngineRootDir(FString &OutRootDir) OVERRIDE;
+	virtual bool GetDefaultEngineIdentifier(FString &OutIdentifier) override;
+	virtual bool GetDefaultEngineRootDir(FString &OutRootDir) override;
 	virtual bool IsPreferredEngineIdentifier(const FString &Identifier, const FString &OtherIdentifier);
 
-	virtual bool IsStockEngineRelease(const FString &Identifier) OVERRIDE;
-	virtual bool IsSourceDistribution(const FString &RootDir) OVERRIDE;
-	virtual bool IsPerforceBuild(const FString &RootDir) OVERRIDE;
-	virtual bool IsValidRootDirectory(const FString &RootDir) OVERRIDE;
+	virtual bool IsStockEngineRelease(const FString &Identifier) override;
+	virtual bool IsSourceDistribution(const FString &RootDir) override;
+	virtual bool IsPerforceBuild(const FString &RootDir) override;
+	virtual bool IsValidRootDirectory(const FString &RootDir) override;
 
-	virtual bool SetEngineIdentifierForProject(const FString &ProjectFileName, const FString &Identifier) OVERRIDE;
-	virtual bool GetEngineIdentifierForProject(const FString &ProjectFileName, FString &OutIdentifier) OVERRIDE;
+	virtual bool SetEngineIdentifierForProject(const FString &ProjectFileName, const FString &Identifier) override;
+	virtual bool GetEngineIdentifierForProject(const FString &ProjectFileName, FString &OutIdentifier) override;
 
-	virtual bool CleanGameProject(const FString& ProjectDir, FFeedbackContext* Warn) OVERRIDE;
-	virtual bool CompileGameProject(const FString& RootDir, const FString& ProjectFileName, FFeedbackContext* Warn) OVERRIDE;
-	virtual bool GenerateProjectFiles(const FString& RootDir, const FString& ProjectFileName, FFeedbackContext* Warn) OVERRIDE;
+	virtual bool CleanGameProject(const FString& ProjectDir, FFeedbackContext* Warn) override;
+	virtual bool CompileGameProject(const FString& RootDir, const FString& ProjectFileName, FFeedbackContext* Warn) override;
+	virtual bool GenerateProjectFiles(const FString& RootDir, const FString& ProjectFileName, FFeedbackContext* Warn) override;
 
-	virtual bool EnumerateProjectsKnownByEngine(const FString &Identifier, bool bIncludeNativeProjects, TArray<FString> &OutProjectFileNames) OVERRIDE;
-	virtual FString GetDefaultProjectCreationPath() OVERRIDE;
+	virtual bool EnumerateProjectsKnownByEngine(const FString &Identifier, bool bIncludeNativeProjects, TArray<FString> &OutProjectFileNames) override;
+	virtual FString GetDefaultProjectCreationPath() override;
 
 private:
 	FString CurrentEngineIdentifier;
diff --git a/Engine/Source/Developer/DesktopPlatform/Private/Linux/DesktopPlatformLinux.h b/Engine/Source/Developer/DesktopPlatform/Private/Linux/DesktopPlatformLinux.h
--- a/Engine/Source/Developer/DesktopPlatform/Private/Linux/DesktopPlatformLinux.h
+++ b/Engine/Source/Developer/DesktopPlatform/Private/Linux/DesktopPlatformLinux.h
@@ -8,21 +8,21 @@ class FDesktopPlatformLinux : public FDesktopPlatformBase
 {
 public:
 	// IDesktopPlatform Implementation
-	virtual bool OpenFileDialog(const void* ParentWindowHandle, const FString& DialogTitle, const FString& DefaultPath, const FString& DefaultFile, const FString& FileTypes, uint32 Flags, TArray<FString>& OutFilenames) OVERRIDE;
-	virtual bool SaveFileDialog(const void* ParentWindowHandle, const FString& DialogTitle, const FString& DefaultPath, const FString& DefaultFile, const FString& FileTypes, uint32 Flags, TArray<FString>& OutFilenames) OVERRIDE;
-	virtual bool OpenDirectoryDialog(const void* ParentWindowHandle, const FString& DialogTitle, const FString& DefaultPath, FString& OutFolderName) OVERRIDE;
-	virtual bool OpenFontDialog(const void* ParentWindowHandle, FString& OutFontName, float& OutHeight, EFontImportFlags::Type& OutFlags) OVERRIDE;
-	virtual bool OpenLauncher(bool Install, FString CommandLineParams ) OVERRIDE;
+	virtual bool OpenFileDialog(const void* ParentWindowHandle, const FString& DialogTitle, const FString& DefaultPath, const FString& DefaultFile, const FString& FileTypes, uint32 Flags, TArray<FString>& OutFilenames) override;
+	virtual bool SaveFileDialog(const void* ParentWindowHandle, const FString& DialogTitle, const FString& DefaultPath, const FString& DefaultFile, const FString& FileTypes, uint32 Flags, TArray<FString>& OutFilenames) override;
+	virtual bool OpenDirectoryDialog(const void* ParentWindowHandle, const FString& DialogTitle, const FString& DefaultPath, FString& OutFolderName) override;
+	virtual bool OpenFontDialog(const void* ParentWindowHandle, FString& OutFontName, float& OutHeight, EFontImportFlags::Type& OutFlags) override;
+	virtual bool OpenLauncher(bool Install, FString CommandLineParams ) override;
 
-	virtual bool RegisterEngineInstallation(const FString &RootDir, FString &OutIdentifier) OVERRIDE;
-	virtual void EnumerateEngineInstallations(TMap<FString, FString> &OutInstallations) OVERRIDE;
+	virtual bool RegisterEngineInstallation(const FString &RootDir, FString &OutIdentifier) override;
+	virtual void EnumerateEngineInstallations(TMap<FString, FString> &OutInstallations) override;
 
-	virtual bool VerifyFileAssociations() OVERRIDE;
-	virtual bool UpdateFileAssociations() OVERRIDE;
+	virtual bool VerifyFileAssociations() override;
+	virtual bool UpdateFileAssociations() override;
 
-	virtual bool RunUnrealBuildTool(const FText& Description, const FString& RootDir, const FString& Arguments, FFeedbackContext* Warn) OVERRIDE;
+	virtual bool RunUnrealBuildTool(const FText& Description, const FString& RootDir, const FString& Arguments, FFeedbackContext* Warn) override;
 
-	virtual FFeedbackContext* GetNativeFeedbackContext() OVERRIDE;
+	virtual FFeedbackContext* GetNativeFeedbackContext() override;
 
 private:
 	bool FileDialogShared(bool bSave, const void* ParentWindowHandle, const FString& DialogTitle, const FString& DefaultPath, const FString& DefaultFile, const FString& FileTypes, uint32 Flags, TArray<FString>& OutFilenames);
diff --git a/Engine/Source/Developer/DesktopPlatform/Private/Mac/DesktopPlatformMac.h b/Engine/Source/Developer/DesktopPlatform/Private/Mac/DesktopPlatformMac.h
--- a/Engine/Source/Developer/DesktopPlatform/Private/Mac/DesktopPlatformMac.h
+++ b/Engine/Source/Developer/DesktopPlatform/Private/Mac/DesktopPlatformMac.h
@@ -8,21 +8,21 @@ class FDesktopPlatformMac : public FDesktopPlatformBase
 {
 public:
 	// IDesktopPlatform Implementation
-	virtual bool OpenFileDialog(const void* ParentWindowHandle, const FString& DialogTitle, const FString& DefaultPath, const FString& DefaultFile, const FString& FileTypes, uint32 Flags, TArray<FString>& OutFilenames) OVERRIDE;
-	virtual bool SaveFileDialog(const void* ParentWindowHandle, const FString& DialogTitle, const FString& DefaultPath, const FString& DefaultFile, const FString& FileTypes, uint32 Flags, TArray<FString>& OutFilenames) OVERRIDE;
-	virtual bool OpenDirectoryDialog(const void* ParentWindowHandle, const FString& DialogTitle, const FString& DefaultPath, FString& OutFolderName) OVERRIDE;
-	virtual bool OpenFontDialog(const void* ParentWindowHandle, FString& OutFontName, float& OutHeight, EFontImportFlags::Type& OutFlags) OVERRIDE;
-	virtual bool OpenLauncher(bool Install, FString CommandLineParams ) OVERRIDE;
+	virtual bool OpenFileDialog(const void* ParentWindowHandle, const FString& DialogTitle, const FString& DefaultPath, const FString& DefaultFile, const FString& FileTypes, uint32 Flags, TArray<FString>& OutFilenames) override;
+	virtual bool SaveFileDialog(const void* ParentWindowHandle, const FString& DialogTitle, const FString& DefaultPath, const FString& DefaultFile, const FString& FileTypes, uint32 Flags, TArray<FString>& OutFilenames) override;
+	virtual bool OpenDirectoryDialog(const void* ParentWindowHandle, const FString& DialogTitle, const FString& DefaultPath, FString& OutFolderName) override;
+	virtual bool OpenFontDialog(const void* ParentWindowHandle, FString& OutFontName, float& OutHeight, EFontImportFlags::Type& OutFlags) override;
+	virtual bool OpenLauncher(bool Install, FString CommandLineParams ) override;
 
-	virtual bool RegisterEngineInstallation(const FString &RootDir, FString &OutIdentifier) OVERRIDE;
-	virtual void EnumerateEngineInstallations(TMap<FString, FString> &OutInstallations) OVERRIDE;
+	virtual bool RegisterEngineInstallation(const FString &RootDir, FString &OutIdentifier) override;
+	virtual void EnumerateEngineInstallations(TMap<FString, FString> &OutInstallations) override;
 
-	virtual bool VerifyFileAssociations() OVERRIDE;
-	virtual bool UpdateFileAssociations() OVERRIDE;
+	virtual bool VerifyFileAssociations() override;
+	virtual bool UpdateFileAssociations() override;
 
-	virtual bool RunUnrealBuildTool(const FText& Description, const FString& RootDir, const FString& Arguments, FFeedbackContext* Warn) OVERRIDE;
+	virtual bool RunUnrealBuildTool(const FText& Description, const FString& RootDir, const FString& Arguments, FFeedbackContext* Warn) override;
 
-	virtual FFeedbackContext* GetNativeFeedbackContext() OVERRIDE;
+	virtual FFeedbackContext* GetNativeFeedbackContext() override;
 
 private:
 	bool FileDialogShared(bool bSave, const void* ParentWindowHandle, const FString& DialogTitle, const FString& DefaultPath, const FString& DefaultFile, const FString& FileTypes, uint32 Flags, TArray<FString>& OutFilenames);
diff --git a/Engine/Source/Developer/DesktopPlatform/Private/Mac/MacNativeFeedbackContext.h b/Engine/Source/Developer/DesktopPlatform/Private/Mac/MacNativeFeedbackContext.h
--- a/Engine/Source/Developer/DesktopPlatform/Private/Mac/MacNativeFeedbackContext.h
+++ b/Engine/Source/Developer/DesktopPlatform/Private/Mac/MacNativeFeedbackContext.h
@@ -16,17 +16,17 @@ public:
 	FMacNativeFeedbackContext();
 	virtual ~FMacNativeFeedbackContext();
 
-	virtual void Serialize( const TCHAR* V, ELogVerbosity::Type Verbosity, const class FName& Category ) OVERRIDE;
+	virtual void Serialize( const TCHAR* V, ELogVerbosity::Type Verbosity, const class FName& Category ) override;
 
-	virtual bool YesNof(const FText& Text) OVERRIDE;
+	virtual bool YesNof(const FText& Text) override;
 
-	virtual bool ReceivedUserCancel() OVERRIDE;
-	virtual void BeginSlowTask( const FText& Task, bool bShowProgressDialog, bool bInShowCancelButton=false ) OVERRIDE;
-	virtual void EndSlowTask() OVERRIDE;
+	virtual bool ReceivedUserCancel() override;
+	virtual void BeginSlowTask( const FText& Task, bool bShowProgressDialog, bool bInShowCancelButton=false ) override;
+	virtual void EndSlowTask() override;
 
-	virtual bool StatusUpdate( int32 Numerator, int32 Denominator, const FText& NewStatus ) OVERRIDE;
-	virtual bool StatusForceUpdate( int32 Numerator, int32 Denominator, const FText& StatusText ) OVERRIDE;
-	virtual void UpdateProgress(int32 Numerator, int32 Denominator) OVERRIDE;
+	virtual bool StatusUpdate( int32 Numerator, int32 Denominator, const FText& NewStatus ) override;
+	virtual bool StatusForceUpdate( int32 Numerator, int32 Denominator, const FText& StatusText ) override;
+	virtual void UpdateProgress(int32 Numerator, int32 Denominator) override;
 
 	FContextSupplier* GetContext() const;
 	void SetContext( FContextSupplier* InSupplier );
diff --git a/Engine/Source/Developer/DesktopPlatform/Private/Windows/DesktopPlatformWindows.h b/Engine/Source/Developer/DesktopPlatform/Private/Windows/DesktopPlatformWindows.h
--- a/Engine/Source/Developer/DesktopPlatform/Private/Windows/DesktopPlatformWindows.h
+++ b/Engine/Source/Developer/DesktopPlatform/Private/Windows/DesktopPlatformWindows.h
@@ -9,23 +9,23 @@ class FDesktopPlatformWindows : public FDesktopPlatformBase
 {
 public:
 	// IDesktopPlatform Implementation
-	virtual bool OpenFileDialog(const void* ParentWindowHandle, const FString& DialogTitle, const FString& DefaultPath, const FString& DefaultFile, const FString& FileTypes, uint32 Flags, TArray<FString>& OutFilenames) OVERRIDE;
-	virtual bool SaveFileDialog(const void* ParentWindowHandle, const FString& DialogTitle, const FString& DefaultPath, const FString& DefaultFile, const FString& FileTypes, uint32 Flags, TArray<FString>& OutFilenames) OVERRIDE;
-	virtual bool OpenDirectoryDialog(const void* ParentWindowHandle, const FString& DialogTitle, const FString& DefaultPath, FString& OutFolderName) OVERRIDE;
-	virtual bool OpenFontDialog(const void* ParentWindowHandle, FString& OutFontName, float& OutHeight, EFontImportFlags::Type& OutFlags) OVERRIDE;
-	virtual bool OpenLauncher(bool Install, FString CommandLineParams ) OVERRIDE;
+	virtual bool OpenFileDialog(const void* ParentWindowHandle, const FString& DialogTitle, const FString& DefaultPath, const FString& DefaultFile, const FString& FileTypes, uint32 Flags, TArray<FString>& OutFilenames) override;
+	virtual bool SaveFileDialog(const void* ParentWindowHandle, const FString& DialogTitle, const FString& DefaultPath, const FString& DefaultFile, const FString& FileTypes, uint32 Flags, TArray<FString>& OutFilenames) override;
+	virtual bool OpenDirectoryDialog(const void* ParentWindowHandle, const FString& DialogTitle, const FString& DefaultPath, FString& OutFolderName) override;
+	virtual bool OpenFontDialog(const void* ParentWindowHandle, FString& OutFontName, float& OutHeight, EFontImportFlags::Type& OutFlags) override;
+	virtual bool OpenLauncher(bool Install, FString CommandLineParams ) override;
 
-	virtual bool RegisterEngineInstallation(const FString &RootDir, FString &OutIdentifier) OVERRIDE;
-	virtual void EnumerateEngineInstallations(TMap<FString, FString> &OutInstallations) OVERRIDE;
+	virtual bool RegisterEngineInstallation(const FString &RootDir, FString &OutIdentifier) override;
+	virtual void EnumerateEngineInstallations(TMap<FString, FString> &OutInstallations) override;
 
-	virtual bool IsSourceDistribution(const FString &RootDir) OVERRIDE;
+	virtual bool IsSourceDistribution(const FString &RootDir) override;
 
-	virtual bool VerifyFileAssociations() OVERRIDE;
-	virtual bool UpdateFileAssociations() OVERRIDE;
+	virtual bool VerifyFileAssociations() override;
+	virtual bool UpdateFileAssociations() override;
 
-	virtual bool RunUnrealBuildTool(const FText& Description, const FString& RootDir, const FString& Arguments, FFeedbackContext* Warn) OVERRIDE;
+	virtual bool RunUnrealBuildTool(const FText& Description, const FString& RootDir, const FString& Arguments, FFeedbackContext* Warn) override;
 
-	virtual FFeedbackContext* GetNativeFeedbackContext() OVERRIDE;
+	virtual FFeedbackContext* GetNativeFeedbackContext() override;
 
 private:
 	bool FileDialogShared(bool bSave, const void* ParentWindowHandle, const FString& DialogTitle, const FString& DefaultPath, const FString& DefaultFile, const FString& FileTypes, uint32 Flags, TArray<FString>& OutFilenames);
diff --git a/Engine/Source/Developer/DesktopPlatform/Private/Windows/WindowsNativeFeedbackContext.h b/Engine/Source/Developer/DesktopPlatform/Private/Windows/WindowsNativeFeedbackContext.h
--- a/Engine/Source/Developer/DesktopPlatform/Private/Windows/WindowsNativeFeedbackContext.h
+++ b/Engine/Source/Developer/DesktopPlatform/Private/Windows/WindowsNativeFeedbackContext.h
@@ -19,17 +19,17 @@ public:
 	FWindowsNativeFeedbackContext();
 	virtual ~FWindowsNativeFeedbackContext();
 
-	virtual void Serialize( const TCHAR* V, ELogVerbosity::Type Verbosity, const class FName& Category ) OVERRIDE;
+	virtual void Serialize( const TCHAR* V, ELogVerbosity::Type Verbosity, const class FName& Category ) override;
 
 	VARARG_BODY( bool, YesNof, const TCHAR*, VARARG_NONE );
 
-	virtual bool ReceivedUserCancel() OVERRIDE;
-	virtual void BeginSlowTask( const FText& Task, bool bShowProgressDialog, bool bInShowCancelButton=false ) OVERRIDE;
-	virtual void EndSlowTask() OVERRIDE;
+	virtual bool ReceivedUserCancel() override;
+	virtual void BeginSlowTask( const FText& Task, bool bShowProgressDialog, bool bInShowCancelButton=false ) override;
+	virtual void EndSlowTask() override;
 
-	virtual bool StatusUpdate( int32 Numerator, int32 Denominator, const FText& NewStatus ) OVERRIDE;
-	virtual bool StatusForceUpdate( int32 Numerator, int32 Denominator, const FText& StatusText ) OVERRIDE;
-	virtual void UpdateProgress(int32 Numerator, int32 Denominator) OVERRIDE;
+	virtual bool StatusUpdate( int32 Numerator, int32 Denominator, const FText& NewStatus ) override;
+	virtual bool StatusForceUpdate( int32 Numerator, int32 Denominator, const FText& StatusText ) override;
+	virtual void UpdateProgress(int32 Numerator, int32 Denominator) override;
 
 	FContextSupplier* GetContext() const;
 	void SetContext( FContextSupplier* InSupplier );
diff --git a/Engine/Source/Developer/DeviceManager/Private/Widgets/Apps/SDeviceAppsAppListRow.h b/Engine/Source/Developer/DeviceManager/Private/Widgets/Apps/SDeviceAppsAppListRow.h
--- a/Engine/Source/Developer/DeviceManager/Private/Widgets/Apps/SDeviceAppsAppListRow.h
+++ b/Engine/Source/Developer/DeviceManager/Private/Widgets/Apps/SDeviceAppsAppListRow.h
@@ -37,7 +37,7 @@ public:
 	 * @param ColumnName The name of the column to generate the widget for.
 	 * @return The widget.
 	 */
-	virtual TSharedRef<SWidget> GenerateWidgetForColumn( const FName& ColumnName ) OVERRIDE
+	virtual TSharedRef<SWidget> GenerateWidgetForColumn( const FName& ColumnName ) override
 	{
 		if (ColumnName == TEXT("Date"))
 		{
diff --git a/Engine/Source/Developer/DirectoryWatcher/Private/Mac/DirectoryWatcherMac.h b/Engine/Source/Developer/DirectoryWatcher/Private/Mac/DirectoryWatcherMac.h
--- a/Engine/Source/Developer/DirectoryWatcher/Private/Mac/DirectoryWatcherMac.h
+++ b/Engine/Source/Developer/DirectoryWatcher/Private/Mac/DirectoryWatcherMac.h
@@ -9,9 +9,9 @@ public:
 	FDirectoryWatcherMac();
 	virtual ~FDirectoryWatcherMac();
 
-	virtual bool RegisterDirectoryChangedCallback (const FString& Directory, const FDirectoryChanged& InDelegate) OVERRIDE;
-	virtual bool UnregisterDirectoryChangedCallback (const FString& Directory, const FDirectoryChanged& InDelegate) OVERRIDE;
-	virtual void Tick (float DeltaSeconds) OVERRIDE;
+	virtual bool RegisterDirectoryChangedCallback (const FString& Directory, const FDirectoryChanged& InDelegate) override;
+	virtual bool UnregisterDirectoryChangedCallback (const FString& Directory, const FDirectoryChanged& InDelegate) override;
+	virtual void Tick (float DeltaSeconds) override;
 
 	/** Map of directory paths to requests */
 	TMap<FString, FDirectoryWatchRequestMac*> RequestMap;
diff --git a/Engine/Source/Developer/DirectoryWatcher/Private/Windows/DirectoryWatcherWindows.h b/Engine/Source/Developer/DirectoryWatcher/Private/Windows/DirectoryWatcherWindows.h
--- a/Engine/Source/Developer/DirectoryWatcher/Private/Windows/DirectoryWatcherWindows.h
+++ b/Engine/Source/Developer/DirectoryWatcher/Private/Windows/DirectoryWatcherWindows.h
@@ -8,9 +8,9 @@ public:
 	FDirectoryWatcherWindows();
 	virtual ~FDirectoryWatcherWindows();
 
-	virtual bool RegisterDirectoryChangedCallback (const FString& Directory, const FDirectoryChanged& InDelegate) OVERRIDE;
-	virtual bool UnregisterDirectoryChangedCallback (const FString& Directory, const FDirectoryChanged& InDelegate) OVERRIDE;
-	virtual void Tick (float DeltaSeconds) OVERRIDE;
+	virtual bool RegisterDirectoryChangedCallback (const FString& Directory, const FDirectoryChanged& InDelegate) override;
+	virtual bool UnregisterDirectoryChangedCallback (const FString& Directory, const FDirectoryChanged& InDelegate) override;
+	virtual void Tick (float DeltaSeconds) override;
 
 	/** Map of directory paths to requests */
 	TMap<FString, FDirectoryWatchRequestWindows*> RequestMap;
diff --git a/Engine/Source/Developer/ExternalImagePicker/Private/ExternalImagePickerModule.cpp b/Engine/Source/Developer/ExternalImagePicker/Private/ExternalImagePickerModule.cpp
--- a/Engine/Source/Developer/ExternalImagePicker/Private/ExternalImagePickerModule.cpp
+++ b/Engine/Source/Developer/ExternalImagePicker/Private/ExternalImagePickerModule.cpp
@@ -10,7 +10,7 @@
 class FExternalImagePickerModule : public IExternalImagePickerModule
 {
 public:
-	virtual TSharedRef<SWidget> MakeEditorWidget(const FExternalImagePickerConfiguration& InConfiguration) OVERRIDE
+	virtual TSharedRef<SWidget> MakeEditorWidget(const FExternalImagePickerConfiguration& InConfiguration) override
 	{
 		return SNew(SExternalImagePicker)
 			.TargetImagePath(InConfiguration.TargetImagePath)
diff --git a/Engine/Source/Developer/FriendsAndChat/Classes/FriendsAndChatStyle.h b/Engine/Source/Developer/FriendsAndChat/Classes/FriendsAndChatStyle.h
--- a/Engine/Source/Developer/FriendsAndChat/Classes/FriendsAndChatStyle.h
+++ b/Engine/Source/Developer/FriendsAndChat/Classes/FriendsAndChatStyle.h
@@ -25,7 +25,7 @@ struct FRIENDSANDCHAT_API FFriendsAndChatStyle
 	/**
 	 * Override widget style function.
 	 */
-	virtual void GetResources( TArray< const FSlateBrush* >& OutBrushes ) const OVERRIDE
+	virtual void GetResources( TArray< const FSlateBrush* >& OutBrushes ) const override
 	{}
 
 	// Holds the widget type name
@@ -35,7 +35,7 @@ struct FRIENDSANDCHAT_API FFriendsAndChatStyle
 	 * Get the type name.
 	 * @return the type name
 	 */
-	virtual const FName GetTypeName() const OVERRIDE { return TypeName; };
+	virtual const FName GetTypeName() const override { return TypeName; };
 
 	/**
 	 * Get the default style.
diff --git a/Engine/Source/Developer/FriendsAndChat/Private/FriendsAndChatModule.cpp b/Engine/Source/Developer/FriendsAndChat/Private/FriendsAndChatModule.cpp
--- a/Engine/Source/Developer/FriendsAndChat/Private/FriendsAndChatModule.cpp
+++ b/Engine/Source/Developer/FriendsAndChat/Private/FriendsAndChatModule.cpp
@@ -15,48 +15,48 @@ class FFriendsAndChatModule
 public:
 
 	// Begin IFriendsAndChatModule interface
-	virtual void Init( bool bAllowJoinGame ) OVERRIDE
+	virtual void Init( bool bAllowJoinGame ) override
 	{
 		FFriendsAndChatManager::Get()->Init( FriendsListNotificationDelegate );
 		FFriendsMessageManager::Get()->Init( FriendsListNotificationDelegate, bAllowJoinGame );
 	}
 
-	virtual void SetUnhandledNotification( TSharedRef< FUniqueNetId > NetID ) OVERRIDE
+	virtual void SetUnhandledNotification( TSharedRef< FUniqueNetId > NetID ) override
 	{
 		FFriendsMessageManager::Get()->SetUnhandledNotification( NetID );
 	}
 
-	virtual FOnFriendsNotification& OnFriendsListNotification() OVERRIDE
+	virtual FOnFriendsNotification& OnFriendsListNotification() override
 	{
 		return FriendsListNotificationDelegate;
 	}
 	
-	virtual void CreateFriendsListWidget( TSharedPtr<const SWidget> ParentWidget, const FFriendsAndChatStyle* InStyle ) OVERRIDE
+	virtual void CreateFriendsListWidget( TSharedPtr<const SWidget> ParentWidget, const FFriendsAndChatStyle* InStyle ) override
 	{
 		FFriendsAndChatManager::Get()->GenerateFriendsWindow( ParentWidget, InStyle );
 	}
 
-	virtual TSharedPtr< SWidget > GenerateFriendsListWidget( const struct FFriendsAndChatStyle* InStyle ) OVERRIDE
+	virtual TSharedPtr< SWidget > GenerateFriendsListWidget( const struct FFriendsAndChatStyle* InStyle ) override
 	{
 		return FFriendsAndChatManager::Get()->GenerateFriendsListWidget( InStyle );
 	}
 
-	virtual void SetInSession( bool bInSession ) OVERRIDE
+	virtual void SetInSession( bool bInSession ) override
 	{
 		FFriendsAndChatManager::Get()->SetInSession( bInSession );
 	}
 
-	virtual void ClearGameInvites() OVERRIDE
+	virtual void ClearGameInvites() override
 	{
 		FFriendsMessageManager::Get()->ClearGameInvites();
 	}
 
-	virtual int32 GetFriendCount() OVERRIDE
+	virtual int32 GetFriendCount() override
 	{
 		return FFriendsAndChatManager::Get()->GetFriendCount();
 	}
 
-	virtual void Logout() OVERRIDE
+	virtual void Logout() override
 	{
 		FFriendsAndChatManager::Get()->Logout();
 		FFriendsMessageManager::Get()->Logout();
@@ -67,13 +67,13 @@ public:
 public:
 
 	// Begin IModuleInterface interface
-	virtual void StartupModule() OVERRIDE
+	virtual void StartupModule() override
 	{
 		// Make sure the singletons are created
  		FFriendsAndChatManager::Get();
 	}
 
-	virtual void ShutdownModule() OVERRIDE
+	virtual void ShutdownModule() override
 	{
 		FFriendsAndChatManager::Shutdown();
 	}
diff --git a/Engine/Source/Developer/FriendsAndChat/Private/UI/Widgets/SChatWindow.cpp b/Engine/Source/Developer/FriendsAndChat/Private/UI/Widgets/SChatWindow.cpp
--- a/Engine/Source/Developer/FriendsAndChat/Private/UI/Widgets/SChatWindow.cpp
+++ b/Engine/Source/Developer/FriendsAndChat/Private/UI/Widgets/SChatWindow.cpp
@@ -30,7 +30,7 @@ void SChatWindow::Construct( const FArguments& InArgs )
 			this->Cursor = EMouseCursor::CardinalCross;
 		}
 
-		virtual EWindowZone::Type GetWindowZoneOverride() const OVERRIDE
+		virtual EWindowZone::Type GetWindowZoneOverride() const override
 		{
 			return EWindowZone::TitleBar;
 		}
diff --git a/Engine/Source/Developer/FriendsAndChat/Private/UI/Widgets/SFriendsList.cpp b/Engine/Source/Developer/FriendsAndChat/Private/UI/Widgets/SFriendsList.cpp
--- a/Engine/Source/Developer/FriendsAndChat/Private/UI/Widgets/SFriendsList.cpp
+++ b/Engine/Source/Developer/FriendsAndChat/Private/UI/Widgets/SFriendsList.cpp
@@ -30,7 +30,7 @@ void SFriendsList::Construct( const FArguments& InArgs )
 			this->Cursor = EMouseCursor::CardinalCross;
 		}
 
-		virtual EWindowZone::Type GetWindowZoneOverride() const OVERRIDE
+		virtual EWindowZone::Type GetWindowZoneOverride() const override
 		{
 			return EWindowZone::TitleBar;
 		}
diff --git a/Engine/Source/Developer/FunctionalTesting/Classes/FuncTestRenderingComponent.h b/Engine/Source/Developer/FunctionalTesting/Classes/FuncTestRenderingComponent.h
--- a/Engine/Source/Developer/FunctionalTesting/Classes/FuncTestRenderingComponent.h
+++ b/Engine/Source/Developer/FunctionalTesting/Classes/FuncTestRenderingComponent.h
@@ -10,10 +10,10 @@ class UFuncTestRenderingComponent : public UPrimitiveComponent
 	GENERATED_UCLASS_BODY()
 
 	// Begin UPrimitiveComponent Interface
-	virtual FPrimitiveSceneProxy* CreateSceneProxy() OVERRIDE;
+	virtual FPrimitiveSceneProxy* CreateSceneProxy() override;
 	// End UPrimitiveComponent Interface
 
 	// Begin USceneComponent Interface
-	virtual FBoxSphereBounds CalcBounds(const FTransform &LocalToWorld) const OVERRIDE;
+	virtual FBoxSphereBounds CalcBounds(const FTransform &LocalToWorld) const override;
 	// End USceneComponent Interface
 };
\ No newline at end of file
diff --git a/Engine/Source/Developer/FunctionalTesting/Classes/FunctionalAITest.h b/Engine/Source/Developer/FunctionalTesting/Classes/FunctionalAITest.h
--- a/Engine/Source/Developer/FunctionalTesting/Classes/FunctionalAITest.h
+++ b/Engine/Source/Developer/FunctionalTesting/Classes/FunctionalAITest.h
@@ -122,15 +122,15 @@ class AFunctionalAITest : public AFunctionalTest
 	virtual bool IsOneOfSpawnedPawns(AActor* Actor);
 
 	// AActor interface begin
-	virtual void BeginPlay() OVERRIDE;
-	virtual void Tick(float DeltaSeconds) OVERRIDE;
+	virtual void BeginPlay() override;
+	virtual void Tick(float DeltaSeconds) override;
 	// AActor interface end
 
-	virtual bool StartTest() OVERRIDE;
-	virtual void FinishTest(TEnumAsByte<EFunctionalTestResult::Type> TestResult, const FString& Message) OVERRIDE;
-	virtual bool WantsToRunAgain() const OVERRIDE;
-	virtual void CleanUp() OVERRIDE;
-	virtual FString GetAdditionalTestFinishedMessage(EFunctionalTestResult::Type TestResult) const OVERRIDE;
+	virtual bool StartTest() override;
+	virtual void FinishTest(TEnumAsByte<EFunctionalTestResult::Type> TestResult, const FString& Message) override;
+	virtual bool WantsToRunAgain() const override;
+	virtual void CleanUp() override;
+	virtual FString GetAdditionalTestFinishedMessage(EFunctionalTestResult::Type TestResult) const override;
 
 protected:
 
diff --git a/Engine/Source/Developer/FunctionalTesting/Classes/FunctionalTest.h b/Engine/Source/Developer/FunctionalTesting/Classes/FunctionalTest.h
--- a/Engine/Source/Developer/FunctionalTesting/Classes/FunctionalTest.h
+++ b/Engine/Source/Developer/FunctionalTesting/Classes/FunctionalTest.h
@@ -100,11 +100,11 @@ public:
 	virtual void CleanUp();
 
 #if WITH_EDITOR
-	void PostEditChangeProperty(struct FPropertyChangedEvent& PropertyChangedEvent) OVERRIDE;
+	void PostEditChangeProperty(struct FPropertyChangedEvent& PropertyChangedEvent) override;
 #endif // WITH_EDITOR
 
 	// AActor interface begin
-	virtual void Tick(float DeltaSeconds) OVERRIDE;
+	virtual void Tick(float DeltaSeconds) override;
 	// AActor interface end
 
 protected:
diff --git a/Engine/Source/Developer/FunctionalTesting/Private/FuncTestManager.h b/Engine/Source/Developer/FunctionalTesting/Private/FuncTestManager.h
--- a/Engine/Source/Developer/FunctionalTesting/Private/FuncTestManager.h
+++ b/Engine/Source/Developer/FunctionalTesting/Private/FuncTestManager.h
@@ -9,13 +9,13 @@ class FFuncTestManager : public IFuncTestManager, public TSharedFromThis<FFuncTe
 public:
 		
 	/** Triggers in sequence all functional tests found on the level.*/
-	virtual void RunAllTestsOnMap(bool bClearLog, bool bRunLooped) OVERRIDE;
+	virtual void RunAllTestsOnMap(bool bClearLog, bool bRunLooped) override;
 	
-	virtual bool IsRunning() const OVERRIDE;
+	virtual bool IsRunning() const override;
 	
-	virtual void SetScript(class UFunctionalTestingManager* NewScript) OVERRIDE;
+	virtual void SetScript(class UFunctionalTestingManager* NewScript) override;
 
-	virtual class UFunctionalTestingManager* GetCurrentScript() OVERRIDE { return TestScript.Get(); }
+	virtual class UFunctionalTestingManager* GetCurrentScript() override { return TestScript.Get(); }
 
 	virtual void SetLooping(const bool bLoop);
 
diff --git a/Engine/Source/Developer/GameplayDebugger/Private/GameplayDebuggerExec.cpp b/Engine/Source/Developer/GameplayDebugger/Private/GameplayDebuggerExec.cpp
--- a/Engine/Source/Developer/GameplayDebugger/Private/GameplayDebuggerExec.cpp
+++ b/Engine/Source/Developer/GameplayDebugger/Private/GameplayDebuggerExec.cpp
@@ -21,7 +21,7 @@ struct FGameplayDebuggerExec : public FSelfRegisteringExec
 	TWeakObjectPtr<AGameplayDebuggingReplicator> GetDebuggingReplicator(UWorld* InWorld);
 
 	// Begin FExec Interface
-	virtual bool Exec(UWorld* Inworld, const TCHAR* Cmd, FOutputDevice& Ar) OVERRIDE;
+	virtual bool Exec(UWorld* Inworld, const TCHAR* Cmd, FOutputDevice& Ar) override;
 	// End FExec Interface
 
 	TWeakObjectPtr<AGameplayDebuggingReplicator> CachedDebuggingReplicator;
diff --git a/Engine/Source/Developer/GameplayDebugger/Private/GameplayDebuggingComponent.cpp b/Engine/Source/Developer/GameplayDebugger/Private/GameplayDebuggingComponent.cpp
--- a/Engine/Source/Developer/GameplayDebugger/Private/GameplayDebuggingComponent.cpp
+++ b/Engine/Source/Developer/GameplayDebugger/Private/GameplayDebuggingComponent.cpp
@@ -41,7 +41,7 @@ public:
 		}
 	}
 
-	virtual void DrawStaticElements(FStaticPrimitiveDrawInterface* PDI) OVERRIDE
+	virtual void DrawStaticElements(FStaticPrimitiveDrawInterface* PDI) override
 	{
 		for (int32 Index = 0; Index < ChildProxies.Num(); ++Index)
 		{
@@ -49,7 +49,7 @@ public:
 		}
 	}
 
-	virtual void DrawDynamicElements(FPrimitiveDrawInterface* PDI, const FSceneView* View) OVERRIDE
+	virtual void DrawDynamicElements(FPrimitiveDrawInterface* PDI, const FSceneView* View) override
 	{
 		for (int32 Index = 0; Index < ChildProxies.Num(); ++Index)
 		{
@@ -57,7 +57,7 @@ public:
 		}
 	}
 
-	virtual void RegisterDebugDrawDelgate() OVERRIDE
+	virtual void RegisterDebugDrawDelgate() override
 	{
 		for (int32 Index = 0; Index < ChildProxies.Num(); ++Index)
 		{
@@ -65,7 +65,7 @@ public:
 		}
 	}
 	
-	virtual void UnregisterDebugDrawDelgate() OVERRIDE
+	virtual void UnregisterDebugDrawDelgate() override
 	{
 		for (int32 Index = 0; Index < ChildProxies.Num(); ++Index)
 		{
@@ -73,7 +73,7 @@ public:
 		}
 	}
 
-	virtual FPrimitiveViewRelevance GetViewRelevance(const FSceneView* View) OVERRIDE
+	virtual FPrimitiveViewRelevance GetViewRelevance(const FSceneView* View) override
 	{
 		FPrimitiveViewRelevance Result;
 		for (int32 Index = 0; Index < ChildProxies.Num(); ++Index)
diff --git a/Engine/Source/Developer/HTML5/HTML5TargetPlatform/Private/HTML5TargetDevice.h b/Engine/Source/Developer/HTML5/HTML5TargetPlatform/Private/HTML5TargetDevice.h
--- a/Engine/Source/Developer/HTML5/HTML5TargetPlatform/Private/HTML5TargetDevice.h
+++ b/Engine/Source/Developer/HTML5/HTML5TargetPlatform/Private/HTML5TargetDevice.h
@@ -51,26 +51,26 @@ public:
 
 	// Begin ITargetDevice interface
 
-	virtual bool Connect( ) OVERRIDE;
+	virtual bool Connect( ) override;
 
-	virtual bool Deploy( const FString& SourceFolder, FString& OutAppId ) OVERRIDE;
+	virtual bool Deploy( const FString& SourceFolder, FString& OutAppId ) override;
 
-	virtual void Disconnect( ) OVERRIDE;
+	virtual void Disconnect( ) override;
 
-	virtual ETargetDeviceTypes::Type GetDeviceType( ) const OVERRIDE
+	virtual ETargetDeviceTypes::Type GetDeviceType( ) const override
 	{
 		return ETargetDeviceTypes::Browser;
 	}
 
-	virtual FTargetDeviceId GetId( ) const OVERRIDE;
+	virtual FTargetDeviceId GetId( ) const override;
 
-	virtual FString GetName( ) const OVERRIDE;
+	virtual FString GetName( ) const override;
 
-	virtual FString GetOperatingSystemName( ) OVERRIDE;
+	virtual FString GetOperatingSystemName( ) override;
 
-	virtual int32 GetProcessSnapshot( TArray<FTargetDeviceProcessInfo>& OutProcessInfos ) OVERRIDE;
+	virtual int32 GetProcessSnapshot( TArray<FTargetDeviceProcessInfo>& OutProcessInfos ) override;
 
-	virtual const class ITargetPlatform& GetTargetPlatform( ) const OVERRIDE
+	virtual const class ITargetPlatform& GetTargetPlatform( ) const override
 	{
 		return TargetPlatform;
 	}
@@ -80,30 +80,30 @@ public:
 		return true;
 	}
 
-	virtual bool IsDefault( ) const OVERRIDE
+	virtual bool IsDefault( ) const override
 	{
 		return true;
 	}
 
-	virtual bool Launch( const FString& AppId, EBuildConfigurations::Type BuildConfiguration, EBuildTargets::Type BuildTarget, const FString& Params, uint32* OutProcessId ) OVERRIDE;
+	virtual bool Launch( const FString& AppId, EBuildConfigurations::Type BuildConfiguration, EBuildTargets::Type BuildTarget, const FString& Params, uint32* OutProcessId ) override;
 
-	virtual bool PowerOff( bool Force ) OVERRIDE;
+	virtual bool PowerOff( bool Force ) override;
 
-	virtual bool PowerOn( ) OVERRIDE;
+	virtual bool PowerOn( ) override;
 
-	virtual bool Reboot( bool bReconnect = false ) OVERRIDE;
+	virtual bool Reboot( bool bReconnect = false ) override;
 
-	virtual bool Run( const FString& ExecutablePath, const FString& Params, uint32* OutProcessId ) OVERRIDE;
+	virtual bool Run( const FString& ExecutablePath, const FString& Params, uint32* OutProcessId ) override;
 
-	virtual bool SupportsFeature( ETargetDeviceFeatures::Type Feature ) const OVERRIDE;
+	virtual bool SupportsFeature( ETargetDeviceFeatures::Type Feature ) const override;
 
-	virtual bool SupportsSdkVersion( const FString& VersionString ) const OVERRIDE;
+	virtual bool SupportsSdkVersion( const FString& VersionString ) const override;
 
-	virtual bool TerminateProcess( const int32 ProcessId ) OVERRIDE;
+	virtual bool TerminateProcess( const int32 ProcessId ) override;
 
-	virtual void SetUserCredentials( const FString & UserName, const FString & UserPassword ) OVERRIDE;
+	virtual void SetUserCredentials( const FString & UserName, const FString & UserPassword ) override;
 
-	virtual bool GetUserCredentials( FString & OutUserName, FString & OutUserPassword ) OVERRIDE;
+	virtual bool GetUserCredentials( FString & OutUserName, FString & OutUserPassword ) override;
 
 	// End ITargetDevice interface
 
diff --git a/Engine/Source/Developer/HTML5/HTML5TargetPlatform/Private/HTML5TargetPlatform.h b/Engine/Source/Developer/HTML5/HTML5TargetPlatform/Private/HTML5TargetPlatform.h
--- a/Engine/Source/Developer/HTML5/HTML5TargetPlatform/Private/HTML5TargetPlatform.h
+++ b/Engine/Source/Developer/HTML5/HTML5TargetPlatform/Private/HTML5TargetPlatform.h
@@ -27,50 +27,50 @@ public:
 
 	// Begin ITargetPlatform interface
 
-	virtual void EnableDeviceCheck(bool OnOff) OVERRIDE {}
+	virtual void EnableDeviceCheck(bool OnOff) override {}
 
-	virtual void GetAllDevices( TArray<ITargetDevicePtr>& OutDevices ) const OVERRIDE;
+	virtual void GetAllDevices( TArray<ITargetDevicePtr>& OutDevices ) const override;
 
-	virtual ECompressionFlags GetBaseCompressionMethod( ) const OVERRIDE;
+	virtual ECompressionFlags GetBaseCompressionMethod( ) const override;
 
-	virtual bool GenerateStreamingInstallManifest(const TMultiMap<FString, int32>& ChunkMap, const TSet<int32>& ChunkIDsInUse) const OVERRIDE
+	virtual bool GenerateStreamingInstallManifest(const TMultiMap<FString, int32>& ChunkMap, const TSet<int32>& ChunkIDsInUse) const override
 	{
 		return true;
 	}
 
-	virtual ITargetDevicePtr GetDefaultDevice( ) const OVERRIDE;
+	virtual ITargetDevicePtr GetDefaultDevice( ) const override;
 
-	virtual ITargetDevicePtr GetDevice( const FTargetDeviceId& DeviceId ) OVERRIDE;
+	virtual ITargetDevicePtr GetDevice( const FTargetDeviceId& DeviceId ) override;
 
-	virtual bool IsRunningPlatform( ) const OVERRIDE;
+	virtual bool IsRunningPlatform( ) const override;
 
 #if WITH_ENGINE
-	virtual void GetAllPossibleShaderFormats( TArray<FName>& OutFormats ) const OVERRIDE;
+	virtual void GetAllPossibleShaderFormats( TArray<FName>& OutFormats ) const override;
 
-	virtual void GetAllTargetedShaderFormats( TArray<FName>& OutFormats ) const OVERRIDE;
+	virtual void GetAllTargetedShaderFormats( TArray<FName>& OutFormats ) const override;
 
-	virtual const class FStaticMeshLODSettings& GetStaticMeshLODSettings() const OVERRIDE;
+	virtual const class FStaticMeshLODSettings& GetStaticMeshLODSettings() const override;
 
-	virtual void GetTextureFormats( const UTexture* InTexture, TArray<FName>& OutFormats ) const OVERRIDE;
+	virtual void GetTextureFormats( const UTexture* InTexture, TArray<FName>& OutFormats ) const override;
 
-	virtual void GetReflectionCaptureFormats( TArray<FName>& OutFormats ) const OVERRIDE
+	virtual void GetReflectionCaptureFormats( TArray<FName>& OutFormats ) const override
 	{
 		OutFormats.Add(FName(TEXT("EncodedHDR")));
 	}
 
-	virtual const struct FTextureLODSettings& GetTextureLODSettings( ) const OVERRIDE;
+	virtual const struct FTextureLODSettings& GetTextureLODSettings( ) const override;
 
-	virtual FName GetWaveFormat( class USoundWave* Wave ) const OVERRIDE;
+	virtual FName GetWaveFormat( class USoundWave* Wave ) const override;
 #endif // WITH_ENGINE
 
 	DECLARE_DERIVED_EVENT(FHTML5TargetPlatform, ITargetPlatform::FOnTargetDeviceDiscovered, FOnTargetDeviceDiscovered);
-	virtual FOnTargetDeviceDiscovered& OnDeviceDiscovered( ) OVERRIDE
+	virtual FOnTargetDeviceDiscovered& OnDeviceDiscovered( ) override
 	{
 		return DeviceDiscoveredEvent;
 	}
 
 	DECLARE_DERIVED_EVENT(FHTML5TargetPlatform, ITargetPlatform::FOnTargetDeviceLost, FOnTargetDeviceLost);
-	virtual FOnTargetDeviceLost& OnDeviceLost( ) OVERRIDE
+	virtual FOnTargetDeviceLost& OnDeviceLost( ) override
 	{
 		return DeviceLostEvent;
 	}
diff --git a/Engine/Source/Developer/HTML5/HTML5TargetPlatform/Private/HTML5TargetPlatformModule.cpp b/Engine/Source/Developer/HTML5/HTML5TargetPlatform/Private/HTML5TargetPlatformModule.cpp
--- a/Engine/Source/Developer/HTML5/HTML5TargetPlatform/Private/HTML5TargetPlatformModule.cpp
+++ b/Engine/Source/Developer/HTML5/HTML5TargetPlatform/Private/HTML5TargetPlatformModule.cpp
@@ -34,7 +34,7 @@ public:
 
 	// Begin ITargetPlatformModule interface
 
-	virtual ITargetPlatform* GetTargetPlatform( ) OVERRIDE
+	virtual ITargetPlatform* GetTargetPlatform( ) override
 	{
 		if (Singleton == NULL)
 		{
@@ -52,7 +52,7 @@ public:
 
 	// Begin IModuleInterface interface
 
-	virtual void StartupModule() OVERRIDE
+	virtual void StartupModule() override
 	{
 		TargetSettings = ConstructObject<UHTML5TargetSettings>(UHTML5TargetSettings::StaticClass(), GetTransientPackage(), "HTML5TargetSettings", RF_Standalone);
 		TargetSettings->AddToRoot();
@@ -69,7 +69,7 @@ public:
 		}
 	}
 
-	virtual void ShutdownModule() OVERRIDE
+	virtual void ShutdownModule() override
 	{
 		ISettingsModule* SettingsModule = ISettingsModule::Get();
 
diff --git a/Engine/Source/Developer/ImageWrapper/Private/PngImageWrapper.h b/Engine/Source/Developer/ImageWrapper/Private/PngImageWrapper.h
--- a/Engine/Source/Developer/ImageWrapper/Private/PngImageWrapper.h
+++ b/Engine/Source/Developer/ImageWrapper/Private/PngImageWrapper.h
@@ -33,13 +33,13 @@ public:
 
 	// Begin FImageWrapper Interface
 
-	virtual void Compress( int32 Quality ) OVERRIDE;
+	virtual void Compress( int32 Quality ) override;
 
-	virtual void Reset( ) OVERRIDE;
+	virtual void Reset( ) override;
 	
-	virtual bool SetCompressed( const void* InCompressedData, int32 InCompressedSize ) OVERRIDE;
+	virtual bool SetCompressed( const void* InCompressedData, int32 InCompressedSize ) override;
 
-	virtual void Uncompress( const ERGBFormat::Type InFormat, int32 InBitDepth ) OVERRIDE;
+	virtual void Uncompress( const ERGBFormat::Type InFormat, int32 InBitDepth ) override;
 	
 	// End FImageWrapper Interface
 
diff --git a/Engine/Source/Developer/InternationalizationSettings/Public/InternationalizationSettingsModelDetails.h b/Engine/Source/Developer/InternationalizationSettings/Public/InternationalizationSettingsModelDetails.h
--- a/Engine/Source/Developer/InternationalizationSettings/Public/InternationalizationSettingsModelDetails.h
+++ b/Engine/Source/Developer/InternationalizationSettings/Public/InternationalizationSettingsModelDetails.h
@@ -20,7 +20,7 @@ public:
 	~FInternationalizationSettingsModelDetails();
 
 	/** IDetailCustomization interface */
-	virtual void CustomizeDetails( IDetailLayoutBuilder& DetailLayout ) OVERRIDE;
+	virtual void CustomizeDetails( IDetailLayoutBuilder& DetailLayout ) override;
 
 private:
 
diff --git a/Engine/Source/Developer/LauncherServices/Private/Launcher/Launcher.h b/Engine/Source/Developer/LauncherServices/Private/Launcher/Launcher.h
--- a/Engine/Source/Developer/LauncherServices/Private/Launcher/Launcher.h
+++ b/Engine/Source/Developer/LauncherServices/Private/Launcher/Launcher.h
@@ -13,7 +13,7 @@ public:
 
 	// Begin ILauncher interface
 
-	virtual ILauncherWorkerPtr Launch( const ITargetDeviceProxyManagerRef& DeviceProxyManager, const ILauncherProfileRef& Profile ) OVERRIDE;
+	virtual ILauncherWorkerPtr Launch( const ITargetDeviceProxyManagerRef& DeviceProxyManager, const ILauncherProfileRef& Profile ) override;
 
 	// End ILauncher interface
 };
diff --git a/Engine/Source/Developer/LauncherServices/Private/Launcher/LauncherBuildCommands.h b/Engine/Source/Developer/LauncherServices/Private/Launcher/LauncherBuildCommands.h
--- a/Engine/Source/Developer/LauncherServices/Private/Launcher/LauncherBuildCommands.h
+++ b/Engine/Source/Developer/LauncherServices/Private/Launcher/LauncherBuildCommands.h
@@ -15,17 +15,17 @@ public:
 		: TargetPlatform(InTargetPlatform)
 	{ }
 
-	virtual FString GetName() const OVERRIDE
+	virtual FString GetName() const override
 	{
 		return NSLOCTEXT("FLauncherTask", "LauncherBuildGameTaskName", "Building game").ToString();
 	}
 
-	virtual FString GetDesc() const OVERRIDE
+	virtual FString GetDesc() const override
 	{
 		return NSLOCTEXT("FLauncherTask", "LauncherBuildGameTaskDesc", "Build game for ").ToString() + TargetPlatform.PlatformName();
 	}
 
-	virtual FString GetArguments(FLauncherTaskChainState& ChainState) const OVERRIDE
+	virtual FString GetArguments(FLauncherTaskChainState& ChainState) const override
 	{
 		FString CommandLine;
 
@@ -53,17 +53,17 @@ public:
 		: TargetPlatform(InTargetPlatform)
 	{ }
 
-	virtual FString GetName() const OVERRIDE
+	virtual FString GetName() const override
 	{
 		return NSLOCTEXT("FLauncherTask", "LauncherBuildServerTaskName", "Building server").ToString();
 	}
 
-	virtual FString GetDesc() const OVERRIDE
+	virtual FString GetDesc() const override
 	{
 		return NSLOCTEXT("FLauncherTask", "LauncherBuildServerTaskDesc", "Build game for ").ToString() + TargetPlatform.PlatformName();
 	}
 
-	virtual FString GetArguments(FLauncherTaskChainState& ChainState) const OVERRIDE
+	virtual FString GetArguments(FLauncherTaskChainState& ChainState) const override
 	{
 		FString CommandLine;
 
diff --git a/Engine/Source/Developer/LauncherServices/Private/Launcher/LauncherCookCommands.h b/Engine/Source/Developer/LauncherServices/Private/Launcher/LauncherCookCommands.h
--- a/Engine/Source/Developer/LauncherServices/Private/Launcher/LauncherCookCommands.h
+++ b/Engine/Source/Developer/LauncherServices/Private/Launcher/LauncherCookCommands.h
@@ -15,17 +15,17 @@ public:
 		: TargetPlatform(InTargetPlatform)
 	{ }
 
-	virtual FString GetName() const OVERRIDE
+	virtual FString GetName() const override
 	{
 		return NSLOCTEXT("FLauncherTask", "LauncherCookTaskName", "Cooking content").ToString();
 	}
 
-	virtual FString GetDesc() const OVERRIDE
+	virtual FString GetDesc() const override
 	{
 		return NSLOCTEXT("FLauncherTask", "LauncherCookTaskDesc", "Cook content for ").ToString() + TargetPlatform.PlatformName();
 	}
 
-	virtual FString GetArguments(FLauncherTaskChainState& ChainState) const OVERRIDE
+	virtual FString GetArguments(FLauncherTaskChainState& ChainState) const override
 	{
 		FString CommandLine;
 
@@ -59,12 +59,12 @@ public:
 		return CommandLine;
 	}
 
-	virtual FString GetDependencyArguments(FLauncherTaskChainState& ChainState) const OVERRIDE
+	virtual FString GetDependencyArguments(FLauncherTaskChainState& ChainState) const override
 	{
 		return TEXT(" -skipcook");
 	}
 
-	virtual FString GetAdditionalArguments(FLauncherTaskChainState& ChainState) const OVERRIDE
+	virtual FString GetAdditionalArguments(FLauncherTaskChainState& ChainState) const override
 	{
 		FString CommandLine = TEXT("");
 
@@ -96,17 +96,17 @@ public:
 		: TargetPlatform(InTargetPlatform)
 	{ }
 
-	virtual FString GetName() const OVERRIDE
+	virtual FString GetName() const override
 	{
 		return NSLOCTEXT("FLauncherTask", "LauncherCookTaskName", "Cooking content").ToString();
 	}
 
-	virtual FString GetDesc() const OVERRIDE
+	virtual FString GetDesc() const override
 	{
 		return NSLOCTEXT("FLauncherTask", "LauncherCookTaskDesc", "Cook content for ").ToString() + TargetPlatform.PlatformName();
 	}
 
-	virtual FString GetArguments(FLauncherTaskChainState& ChainState) const OVERRIDE
+	virtual FString GetArguments(FLauncherTaskChainState& ChainState) const override
 	{
 		FString CommandLine;
 
@@ -150,12 +150,12 @@ public:
 		return CommandLine;
 	}
 
-	virtual FString GetDependencyArguments(FLauncherTaskChainState& ChainState) const OVERRIDE
+	virtual FString GetDependencyArguments(FLauncherTaskChainState& ChainState) const override
 	{
 		return TEXT(" -skipcook");
 	}
 
-	virtual FString GetAdditionalArguments(FLauncherTaskChainState& ChainState) const OVERRIDE
+	virtual FString GetAdditionalArguments(FLauncherTaskChainState& ChainState) const override
 	{
 		FString CommandLine = TEXT("");
 
@@ -198,17 +198,17 @@ public:
 		}
 	}
 
-	virtual FString GetName() const OVERRIDE
+	virtual FString GetName() const override
 	{
 		return NSLOCTEXT("FLauncherTask", "LauncherDeployFileServerTaskName", "Starting file server").ToString();
 	}
 
-	virtual FString GetDesc() const OVERRIDE
+	virtual FString GetDesc() const override
 	{
 		return NSLOCTEXT("FLauncherTask", "LauncherDeployFileServerTaskDesc", "Start file server for ").ToString() + TargetPlatform.PlatformName();
 	}
 
-	virtual FString GetArguments(FLauncherTaskChainState& ChainState) const OVERRIDE
+	virtual FString GetArguments(FLauncherTaskChainState& ChainState) const override
 	{
 		FString CommandLine;
 
@@ -230,7 +230,7 @@ public:
 		return CommandLine;
 	}
 
-	virtual FString GetDependencyArguments(FLauncherTaskChainState& ChainState) const OVERRIDE
+	virtual FString GetDependencyArguments(FLauncherTaskChainState& ChainState) const override
 	{
 		FString CommandLine;
 		CommandLine += TEXT(" -cookonthefly -skipserver");
@@ -247,7 +247,7 @@ public:
 		return CommandLine;
 	}
 
-	virtual FString GetAdditionalArguments(FLauncherTaskChainState& ChainState) const OVERRIDE
+	virtual FString GetAdditionalArguments(FLauncherTaskChainState& ChainState) const override
 	{
 		FString CommandLine;
 
@@ -260,12 +260,12 @@ public:
 		return CommandLine;
 	}
 
-	virtual bool IsComplete() const OVERRIDE
+	virtual bool IsComplete() const override
 	{
 		return StartDetected;
 	}
 
-	virtual bool PostExecute(FLauncherTaskChainState& ChainState) OVERRIDE
+	virtual bool PostExecute(FLauncherTaskChainState& ChainState) override
 	{
 		if (StartDetected)
 		{
@@ -322,17 +322,17 @@ public:
 		: TargetPlatform(InTargetPlatform)
 	{ }
 
-	virtual FString GetName() const OVERRIDE
+	virtual FString GetName() const override
 	{
 		return NSLOCTEXT("FLauncherTask", "LauncherDeployFileServerTaskName", "Starting file server").ToString();
 	}
 
-	virtual FString GetDesc() const OVERRIDE
+	virtual FString GetDesc() const override
 	{
 		return NSLOCTEXT("FLauncherTask", "LauncherDeployFileServerTaskDesc", "Start file server for ").ToString() + TargetPlatform.PlatformName();
 	}
 
-	virtual FString GetArguments(FLauncherTaskChainState& ChainState) const OVERRIDE
+	virtual FString GetArguments(FLauncherTaskChainState& ChainState) const override
 	{
 		FString CommandLine;
 
@@ -346,7 +346,7 @@ public:
 		return CommandLine;
 	}
 
-	virtual FString GetDependencyArguments(FLauncherTaskChainState& ChainState) const OVERRIDE
+	virtual FString GetDependencyArguments(FLauncherTaskChainState& ChainState) const override
 	{
 		FString CommandLine;
 		CommandLine += TEXT(" -cookonthefly -skipserver");
diff --git a/Engine/Source/Developer/LauncherServices/Private/Launcher/LauncherDeployCommands.h b/Engine/Source/Developer/LauncherServices/Private/Launcher/LauncherDeployCommands.h
--- a/Engine/Source/Developer/LauncherServices/Private/Launcher/LauncherDeployCommands.h
+++ b/Engine/Source/Developer/LauncherServices/Private/Launcher/LauncherDeployCommands.h
@@ -18,17 +18,17 @@ public:
 		, LauncherCommandLine(InCmdLine)
 	{ }
 
-	virtual FString GetName() const OVERRIDE
+	virtual FString GetName() const override
 	{
 		return NSLOCTEXT("FLauncherTask", "LauncherDeployTaskName", "Deploying content").ToString();
 	}
 
-	virtual FString GetDesc() const OVERRIDE
+	virtual FString GetDesc() const override
 	{
 		return NSLOCTEXT("FLauncherTask", "LauncherDeployTaskDesc", "Deploying content for ").ToString() + TargetPlatform.PlatformName();
 	}
 
-	virtual FString GetArguments(FLauncherTaskChainState& ChainState) const OVERRIDE
+	virtual FString GetArguments(FLauncherTaskChainState& ChainState) const override
 	{
 		// build UAT command line parameters
 		FString StagePath;
@@ -77,14 +77,14 @@ public:
 		return CommandLine;
 	}
 
-	virtual bool PreExecute(FLauncherTaskChainState& ChainState) OVERRIDE
+	virtual bool PreExecute(FLauncherTaskChainState& ChainState) override
 	{
 		// disable the device check
 		const_cast<ITargetPlatform&>(TargetPlatform).EnableDeviceCheck(false);
 		return true;
 	}
 
-	virtual bool PostExecute(FLauncherTaskChainState& ChainState) OVERRIDE
+	virtual bool PostExecute(FLauncherTaskChainState& ChainState) override
 	{
 		// disable the device check
 		const_cast<ITargetPlatform&>(TargetPlatform).EnableDeviceCheck(true);
@@ -125,17 +125,17 @@ public:
 		, CookCommand(InCook)
 	{ }
 
-	virtual FString GetName() const OVERRIDE
+	virtual FString GetName() const override
 	{
 		return NSLOCTEXT("FLauncherTask", "LauncherDeployTaskName", "Deploying content").ToString();
 	}
 
-	virtual FString GetDesc() const OVERRIDE
+	virtual FString GetDesc() const override
 	{
 		return NSLOCTEXT("FLauncherTask", "LauncherDeployTaskDesc", "Deploying content for ").ToString() + TargetPlatform.PlatformName();
 	}
 
-	virtual FString GetArguments(FLauncherTaskChainState& ChainState) const OVERRIDE
+	virtual FString GetArguments(FLauncherTaskChainState& ChainState) const override
 	{
 		// build UAT command line parameters
 		FString StagePath = FPaths::ConvertRelativePathToFull(ChainState.Profile->GetProjectBasePath() + FString(TEXT("StagedBuilds")));
@@ -222,17 +222,17 @@ public:
 		, LauncherCommandLine(InCmdLine)
 	{ }
 
-	virtual FString GetName() const OVERRIDE
+	virtual FString GetName() const override
 	{
 		return NSLOCTEXT("FLauncherTask", "LauncherDeployTaskName", "Deploying content").ToString();
 	}
 
-	virtual FString GetDesc() const OVERRIDE
+	virtual FString GetDesc() const override
 	{
 		return NSLOCTEXT("FLauncherTask", "LauncherDeployTaskDesc", "Deploying content for ").ToString() + TargetPlatform.PlatformName();
 	}
 
-	virtual FString GetArguments(FLauncherTaskChainState& ChainState) const OVERRIDE
+	virtual FString GetArguments(FLauncherTaskChainState& ChainState) const override
 	{
 		// build UAT command line parameters
 		FString CommandLine;
@@ -268,14 +268,14 @@ public:
 		return CommandLine;
 	}
 
-	virtual bool PreExecute(FLauncherTaskChainState& ChainState) OVERRIDE
+	virtual bool PreExecute(FLauncherTaskChainState& ChainState) override
 	{
 		// disable the device check
 		const_cast<ITargetPlatform&>(TargetPlatform).EnableDeviceCheck(false);
 		return true;
 	}
 
-	virtual bool PostExecute(FLauncherTaskChainState& ChainState) OVERRIDE
+	virtual bool PostExecute(FLauncherTaskChainState& ChainState) override
 	{
 		// disable the device check
 		const_cast<ITargetPlatform&>(TargetPlatform).EnableDeviceCheck(true);
@@ -316,17 +316,17 @@ public:
 		, CookCommand(InCook)
 	{ }
 
-	virtual FString GetName() const OVERRIDE
+	virtual FString GetName() const override
 	{
 		return NSLOCTEXT("FLauncherTask", "LauncherDeployTaskName", "Deploying content").ToString();
 	}
 
-	virtual FString GetDesc() const OVERRIDE
+	virtual FString GetDesc() const override
 	{
 		return NSLOCTEXT("FLauncherTask", "LauncherDeployTaskDesc", "Deploying content for ").ToString() + TargetPlatform.PlatformName();
 	}
 
-	virtual FString GetArguments(FLauncherTaskChainState& ChainState) const OVERRIDE
+	virtual FString GetArguments(FLauncherTaskChainState& ChainState) const override
 	{
 		// build UAT command line parameters
 		FString CommandLine;
diff --git a/Engine/Source/Developer/LauncherServices/Private/Launcher/LauncherLaunchCommands.h b/Engine/Source/Developer/LauncherServices/Private/Launcher/LauncherLaunchCommands.h
--- a/Engine/Source/Developer/LauncherServices/Private/Launcher/LauncherLaunchCommands.h
+++ b/Engine/Source/Developer/LauncherServices/Private/Launcher/LauncherLaunchCommands.h
@@ -19,17 +19,17 @@ public:
 		, CookCommand(InCook)
 	{ }
 
-	virtual FString GetName() const OVERRIDE
+	virtual FString GetName() const override
 	{
 		return FString::Printf(*NSLOCTEXT("FLauncherTask", "LauncherTaskGame", "Launching Game").ToString());
 	}
 
-	virtual FString GetDesc() const OVERRIDE
+	virtual FString GetDesc() const override
 	{
 		return FString::Printf(*NSLOCTEXT("FLauncherTask", "LauncherTaskDesc", "Launch on %s as %s").ToString(), *DeviceProxy->GetName(), *Role->GetName());
 	}
 
-	virtual FString GetArguments(FLauncherTaskChainState& ChainState) const OVERRIDE
+	virtual FString GetArguments(FLauncherTaskChainState& ChainState) const override
 	{
 		FString CommandLine;
 
@@ -109,17 +109,17 @@ public:
 		, CookCommand(InCook)
 	{ }
 
-	virtual FString GetName() const OVERRIDE
+	virtual FString GetName() const override
 	{
 		return FString::Printf(*NSLOCTEXT("FLauncherTask", "LauncherTaskServer", "Launching Server").ToString());
 	}
 
-	virtual FString GetDesc() const OVERRIDE
+	virtual FString GetDesc() const override
 	{
 		return FString::Printf(*NSLOCTEXT("FLauncherTask", "LauncherTaskDesc", "Launch on %s as %s").ToString(), *DeviceProxy->GetName(), *Role->GetName());
 	}
 
-	virtual FString GetArguments(FLauncherTaskChainState& ChainState) const OVERRIDE
+	virtual FString GetArguments(FLauncherTaskChainState& ChainState) const override
 	{
 		FString CommandLine;
 
