Minor clang fixes (initialization order issues).

diff --git a/Engine/Source/Developer/AutomationController/Private/AutomationReport.cpp b/Engine/Source/Developer/AutomationController/Private/AutomationReport.cpp
--- a/Engine/Source/Developer/AutomationController/Private/AutomationReport.cpp
+++ b/Engine/Source/Developer/AutomationController/Private/AutomationReport.cpp
@@ -9,10 +9,10 @@
 
 FAutomationReport::FAutomationReport(FAutomationTestInfo& InTestInfo, bool InIsParent)
 	: bEnabled( false )
-	, bNodeExpandInUI( false )
-	, bSelfPassesFilter( false )
-	, bIsParent( InIsParent )
-	, SupportFlags( 0 )
+	, bIsParent(InIsParent)
+	, bNodeExpandInUI(false)
+	, bSelfPassesFilter(false)
+	, SupportFlags(0)
 	, TestInfo( InTestInfo )
 	, bTrackingHistory(false)
 	, NumRecordsToKeep(0)
diff --git a/Engine/Source/Developer/ScreenShotComparisonTools/Private/ScreenShotBaseNode.h b/Engine/Source/Developer/ScreenShotComparisonTools/Private/ScreenShotBaseNode.h
--- a/Engine/Source/Developer/ScreenShotComparisonTools/Private/ScreenShotBaseNode.h
+++ b/Engine/Source/Developer/ScreenShotComparisonTools/Private/ScreenShotBaseNode.h
@@ -22,8 +22,8 @@ public:
 	* @param InName  The Name of this test data item.
 	*/
 	FScreenShotBaseNode( const FString& InName, const FString& InAssetName = "" )
-		: ItemName(InName)
-		, AssetName(InAssetName)
+		: AssetName(InAssetName)
+		, ItemName(InName)
 	{ }
 
 public:
