Code and documentation cleanup pass

- removed dummy UClasses (no longer needed)
- removed file header comments (not used)
- removed duplicated function documentation in cpp files
- documentation cleanup, punctuation, spelling etc.
- pragma once include guards (now work on all platforms)
- relative public includes (are auto-discovered by UBT)
- fixed too many/too few line breaks
- deleted empty files
- missing override
- NULL to nullptr

diff --git a/Engine/Plugins/Messaging/UdpMessaging/Source/UdpMessaging/Private/Transport/ReassembledUdpMessage.h b/Engine/Plugins/Messaging/UdpMessaging/Source/UdpMessaging/Private/Transport/ReassembledUdpMessage.h
--- a/Engine/Plugins/Messaging/UdpMessaging/Source/UdpMessaging/Private/Transport/ReassembledUdpMessage.h
+++ b/Engine/Plugins/Messaging/UdpMessaging/Source/UdpMessaging/Private/Transport/ReassembledUdpMessage.h
@@ -17,9 +17,7 @@ class FReassembledUdpMessage
 {
 public:
 
-	/**
-	 * Default constructor.
-	 */
+	/** Default constructor. */
 	FReassembledUdpMessage( ) { }
 
 	/**
@@ -40,9 +38,7 @@ public:
 		Data.AddUninitialized(MessageSize);
 	}
 
-	/**
-	 * Virtual destructor.
-	 */
+	/** Virtual destructor. */
 	virtual ~FReassembledUdpMessage( ) { }
 
 public:
@@ -173,27 +169,27 @@ public:
 
 private:
 
-	// Holds the message data.
+	/** Holds the message data. */
 	TArray<uint8> Data;
 
-	// Holds the time at which the last segment was received.
+	/** Holds the time at which the last segment was received. */
 	FDateTime LastSegmentTime;
 
-	// Holds an array of bits that indicate which segments still need to be received.
+	/** Holds an array of bits that indicate which segments still need to be received. */
 	TBitArray<> PendingSegments;
 
-	// Holds the number of segments that haven't been received yet.
+	/** Holds the number of segments that haven't been received yet. */
 	uint16 PendingSegmentsCount;
 
-	// Holds the number of bytes received so far.
+	/** Holds the number of bytes received so far. */
 	int32 ReceivedBytes;
 
-	// Holds the number of retransmit requests that were sent since the last segment was received.
+	/** Holds the number of retransmit requests that were sent since the last segment was received. */
 	int32 RetransmitRequestsCount;
 
-	// Holds the sender.
+	/** Holds the sender. */
 	FIPv4Endpoint Sender;
 
-	// Holds the message sequence.
+	/** Holds the message sequence. */
 	uint64 Sequence;
 };
diff --git a/Engine/Plugins/Messaging/UdpMessaging/Source/UdpMessaging/Public/Interfaces/IUdpMessageTunnel.h b/Engine/Plugins/Messaging/UdpMessaging/Source/UdpMessaging/Public/Interfaces/IUdpMessageTunnel.h
--- a/Engine/Plugins/Messaging/UdpMessaging/Source/UdpMessaging/Public/Interfaces/IUdpMessageTunnel.h
+++ b/Engine/Plugins/Messaging/UdpMessaging/Source/UdpMessaging/Public/Interfaces/IUdpMessageTunnel.h
@@ -78,13 +78,13 @@ public:
 
 	/**
 	 * Gets a delegate that is executed when the list of incoming connections changed.
+	 *
+	 * @return The delegate.
 	 */
 	virtual FSimpleDelegate& OnConnectionsChanged( ) = 0;
 
 public:
 
-	/**
-	 * Virtual destructor.
-	 */
+	/** Virtual destructor. */
 	virtual ~IUdpMessageTunnel( ) { }
 };
diff --git a/Engine/Plugins/Messaging/UdpMessaging/Source/UdpMessaging/Public/Interfaces/IUdpMessageTunnelConnection.h b/Engine/Plugins/Messaging/UdpMessaging/Source/UdpMessaging/Public/Interfaces/IUdpMessageTunnelConnection.h
--- a/Engine/Plugins/Messaging/UdpMessaging/Source/UdpMessaging/Public/Interfaces/IUdpMessageTunnelConnection.h
+++ b/Engine/Plugins/Messaging/UdpMessaging/Source/UdpMessaging/Public/Interfaces/IUdpMessageTunnelConnection.h
@@ -17,9 +17,7 @@ class IUdpMessageTunnelConnection
 {
 public:
 
-	/**
-	 * Closes this connection.
-	 */
+	/** Closes this connection. */
 	virtual void Close( ) = 0;
 
 	/**
@@ -59,8 +57,6 @@ public:
 
 public:
 
-	/**
-	 * Virtual destructor.
-	 */
+	/** Virtual destructor. */
 	virtual ~IUdpMessageTunnelConnection( ) { }
 };
diff --git a/Engine/Source/Developer/AutomationController/Private/AutomationControllerManager.h b/Engine/Source/Developer/AutomationController/Private/AutomationControllerManager.h
--- a/Engine/Source/Developer/AutomationController/Private/AutomationControllerManager.h
+++ b/Engine/Source/Developer/AutomationController/Private/AutomationControllerManager.h
@@ -1,9 +1,5 @@
 // Copyright 1998-2014 Epic Games, Inc. All Rights Reserved.
 
-/*=============================================================================
-	AutomationControllerModule.h: Declares the FAutomationControllerModule class.
-=============================================================================*/
-
 #pragma once
 
 
@@ -15,20 +11,16 @@ class FAutomationControllerManager
 {
 public:
 
-	// Begin IAutomationController Interface
+public:
+
+	// IAutomationController Interface
 
 	virtual void RequestAvailableWorkers( const FGuid& InSessionId ) override;
-
 	virtual void RequestTests() override;
-
 	virtual void RunTests( const bool bIsLocalSession) override;
-
 	virtual void StopTests() override;
-
 	virtual void Init() override;
-
 	virtual void RequestLoadAsset( const FString& InAssetName ) override;
-
 	virtual void Tick() override;
 
 	virtual void SetNumPasses(const int32 InNumPasses) override
@@ -167,13 +159,9 @@ public:
 	}
 
 	virtual const bool ExportReport(uint32 FileExportTypeMask) override;
-
 	virtual bool IsTestRunnable( IAutomationReportPtr InReport ) const override;
-
 	virtual void RemoveCallbacks() override;
-
 	virtual void Shutdown() override;
-
 	virtual void Startup() override;
 
 	virtual FOnAutomationControllerManagerShutdown& OnShutdown( ) override
@@ -191,11 +179,9 @@ public:
 		return TestsRefreshedDelegate;
 	}
 
-	virtual bool IsDeviceGroupFlagSet( EAutomationDeviceGroupTypes::Type InDeviceGroup ) const;
-
-	virtual void ToggleDeviceGroupFlag( EAutomationDeviceGroupTypes::Type InDeviceGroup );
-
-	void UpdateDeviceGroups();
+	virtual bool IsDeviceGroupFlagSet( EAutomationDeviceGroupTypes::Type InDeviceGroup ) const override;
+	virtual void ToggleDeviceGroupFlag( EAutomationDeviceGroupTypes::Type InDeviceGroup ) override;
+	virtual void UpdateDeviceGroups() override;
 
 	virtual void SetTestsCompleteCallback(const FOnAutomationControllerTestsComplete& NewCallback) override
 	{
@@ -203,102 +189,86 @@ public:
 	}
 
 	virtual void TrackReportHistory(const bool bShouldTrack, const int32 NumReportsToTrack) override;
-
 	virtual const bool IsTrackingHistory() const override;
-
 	virtual const int32 GetNumberHistoryItemsTracking() const override;
-	// End IAutomationController Interface
-
 
 protected:
 
 	/**
 	 * Adds a ping result from a running test.
 	 *
-	 * @param ResponderAddress - The address of the message endpoint that responded to a ping.
+	 * @param ResponderAddress The address of the message endpoint that responded to a ping.
 	 */
 	void AddPingResult( const FMessageAddress& ResponderAddress );
 
 	/**
 	 * Checks the child result.
 	 *
-	 * @param InReport - The child result to check.
+	 * @param InReport The child result to check.
 	 */
 	void CheckChildResult( TSharedPtr< IAutomationReport > InReport );
 
 	/**
-	 * Execute the next task thats available
+	 * Execute the next task thats available.
 	 *
-	 * @param ClusterIndex			The Cluster index of the device type we intend to use
-	 * @param bAllTestsCompleted	Whether all tests have been completed
+	 * @param ClusterIndex The Cluster index of the device type we intend to use.
+	 * @param bAllTestsCompleted Whether all tests have been completed.
 	 */
 	void ExecuteNextTask( int32 ClusterIndex, OUT bool& bAllTestsCompleted );
 
-	/**
-	 * Distributes any tests that are pending and deal with tests finishing.
-	 */
+	/** Distributes any tests that are pending and deal with tests finishing. */
 	void ProcessAvailableTasks();
 
-	/**
-	 * Processes the results after tests are complete.
-	 */
+	/** Processes the results after tests are complete. */
 	void ProcessResults();
 
 	/**
 	 * Removes the test info.
 	 *
-	 * @param TestToRemove - The test to remove.
+	 * @param TestToRemove The test to remove.
 	 */
 	void RemoveTestRunning( const FMessageAddress& TestToRemove );
 
-	/**
-	 * Resets the data holders which have been used to generate the tests available from a worker.
-	 */
+	/** Resets the data holders which have been used to generate the tests available from a worker. */
 	void ResetIntermediateTestData()
 	{
 		TestInfo.Empty();
 	}
 
-	/**
-	 * Changes the controller state.
-	 */
+	/** Changes the controller state. */
 	void SetControllerStatus( EAutomationControllerModuleState::Type AutomationTestState );
 
-	/**
-	 * stores the tests that are valid for a particular device classification.
-	 */
+	/** stores the tests that are valid for a particular device classification. */
 	void SetTestNames(const FMessageAddress& AutomationWorkerAddress);
 
-	/**
-	 * Updates the tests to ensure they are all still running.
-	 */
+	/** Updates the tests to ensure they are all still running. */
 	void UpdateTests();
 
 private:
 
-	// Handles FAutomationWorkerFindWorkersResponse messages.
+	/** Handles FAutomationWorkerFindWorkersResponse messages. */
 	void HandleFindWorkersResponseMessage( const FAutomationWorkerFindWorkersResponse& Message, const IMessageContextRef& Context );
 
-	// Handles FAutomationWorkerPong messages.
+	/** Handles FAutomationWorkerPong messages. */
 	void HandlePongMessage( const FAutomationWorkerPong& Message, const IMessageContextRef& Context );
 
-	// Handles FAutomationWorkerScreenImage messages..
+	/** Handles FAutomationWorkerScreenImage messages. */
 	void HandleReceivedScreenShot( const FAutomationWorkerScreenImage& Message, const IMessageContextRef& Context );
 
-	// Handles FAutomationWorkerRequestNextNetworkCommand messages.
+	/** Handles FAutomationWorkerRequestNextNetworkCommand messages. */
 	void HandleRequestNextNetworkCommandMessage( const FAutomationWorkerRequestNextNetworkCommand& Message, const IMessageContextRef& Context );
 
-	// Handles FAutomationWorkerRequestTestsReply messages.
+	/** Handles FAutomationWorkerRequestTestsReply messages. */
 	void HandleRequestTestsReplyMessage( const FAutomationWorkerRequestTestsReply& Message, const IMessageContextRef& Context );
 
-	// Handles FAutomationWorkerRunTestsReply messages.
+	/** Handles FAutomationWorkerRunTestsReply messages. */
 	void HandleRunTestsReplyMessage( const FAutomationWorkerRunTestsReply& Message, const IMessageContextRef& Context );
 
-	// Handles FAutomationWorkerWorkerOffline messages.
+	/** Handles FAutomationWorkerWorkerOffline messages. */
 	void HandleWorkerOfflineMessage( const FAutomationWorkerWorkerOffline& Message, const IMessageContextRef& Context );
 
-
 private:
+
 	/** Session this controller is currently communicating with */
 	FGuid ActiveSessionId;
 
@@ -350,16 +320,16 @@ private:
 	/** Counter for number of workers we have received responses from for Refreshing the Test List */
 	uint32 RefreshTestResponses;
 
-	/**available stats/status for all tests*/
+	/** Available stats/status for all tests. */
 	FAutomationReportManager ReportManager;
 
-	// The number we are expecting to receive from a worker
+	/** The number we are expecting to receive from a worker */
 	int32 NumOfTestsToReceive;
 
-	// The collection of test data we are to send to a controller
+	/** The collection of test data we are to send to a controller. */
 	TArray< FAutomationTestInfo > TestInfo;
 
-	/** A data holder to keep track of how long tests have been running */
+	/** A data holder to keep track of how long tests have been running. */
 	struct FTestRunningInfo
 	{
 		FTestRunningInfo( FMessageAddress InMessageAddress ):
@@ -373,38 +343,38 @@ private:
 		float LastPingTime;
 	};
 
-	/** A array of running tests */
+	/** A array of running tests. */
 	TArray< FTestRunningInfo > TestRunningArray;
 
-	/** The number of test passes to perform */
+	/** The number of test passes to perform. */
 	int32 NumTestPasses;
 
-	/** The current test pass we are on */
+	/** The current test pass we are on. */
 	int32 CurrentTestPass;
 
-	/** If screenshots are enabled */
+	/** If screenshots are enabled. */
 	bool bScreenshotsEnabled;
 
-	/** If we should request full screen screen shots */
+	/** If we should request full screen screen shots. */
 	bool bRequestFullScreenScreenshots;
 
-	/** If we should track any test history for the next run */
+	/** If we should track any test history for the next run. */
 	bool bTrackHistory;
 
-	/** The number of history items we wish to track */
+	/** The number of history items we wish to track. */
 	int32 NumberOfHistoryItemsTracked;
 
 private:
 
-	// Holds a delegate that is invoked when the controller shuts down.
+	/** Holds a delegate that is invoked when the controller shuts down. */
 	FOnAutomationControllerManagerShutdown ShutdownDelegate;
 
-	// Holds a delegate that is invoked when the controller has tests available.
+	/** Holds a delegate that is invoked when the controller has tests available. */
 	FOnAutomationControllerManagerTestsAvailable TestsAvailableDelegate;
 
-	// Holds a delegate that is invoked when the controller's tests are being refreshed.
+	/** Holds a delegate that is invoked when the controller's tests are being refreshed. */
 	FOnAutomationControllerTestsRefreshed TestsRefreshedDelegate;
 
-	// Holds a delegate that is invoked when the tests have completed
+	/** Holds a delegate that is invoked when the tests have completed. */
 	FOnAutomationControllerTestsComplete TestsCompleteDelegate;
 };
diff --git a/Engine/Source/Developer/AutomationController/Private/AutomationControllerModule.h b/Engine/Source/Developer/AutomationController/Private/AutomationControllerModule.h
--- a/Engine/Source/Developer/AutomationController/Private/AutomationControllerModule.h
+++ b/Engine/Source/Developer/AutomationController/Private/AutomationControllerModule.h
@@ -1,9 +1,5 @@
 // Copyright 1998-2014 Epic Games, Inc. All Rights Reserved.
 
-/*=============================================================================
-	AutomationControllerModule.h: Declares the FAutomationControllerModule class.
-=============================================================================*/
-
 #pragma once
 
 
@@ -15,32 +11,22 @@ class FAutomationControllerModule
 {
 public:
 
-	// Begin IAutomationControllerModule Interface
+	// IAutomationControllerModule interface
 
 	virtual IAutomationControllerManagerRef GetAutomationController( ) override;
-
 	virtual void Init() override;
-
 	virtual void Tick() override;
 
-	// End IAutomationControllerModule Interface
-
-
 public:
 
-	// Begin IModuleInterface interface
+	// IModuleInterface interface
 
 	virtual void StartupModule() override;
-
 	virtual void ShutdownModule() override;
-
 	virtual bool SupportsDynamicReloading() override;
 
-	// End IModuleInterface interface
-
-
 private:
 
-	// Holds the automation controller singleton.
+	/** Holds the automation controller singleton. */
 	static IAutomationControllerManagerPtr AutomationControllerSingleton;
 };
diff --git a/Engine/Source/Developer/AutomationController/Private/AutomationControllerPrivatePCH.h b/Engine/Source/Developer/AutomationController/Private/AutomationControllerPrivatePCH.h
--- a/Engine/Source/Developer/AutomationController/Private/AutomationControllerPrivatePCH.h
+++ b/Engine/Source/Developer/AutomationController/Private/AutomationControllerPrivatePCH.h
@@ -1,9 +1,5 @@
 // Copyright 1998-2014 Epic Games, Inc. All Rights Reserved.
 
-/*=============================================================================
-	AutomationControllerPrivatePCH.h: Pre-compiled header file for the AutomationController module.
-=============================================================================*/
-
 #pragma once
 
 
diff --git a/Engine/Source/Developer/AutomationController/Private/AutomationDeviceClusterManager.cpp b/Engine/Source/Developer/AutomationController/Private/AutomationDeviceClusterManager.cpp
--- a/Engine/Source/Developer/AutomationController/Private/AutomationDeviceClusterManager.cpp
+++ b/Engine/Source/Developer/AutomationController/Private/AutomationDeviceClusterManager.cpp
@@ -1,9 +1,5 @@
 // Copyright 1998-2014 Epic Games, Inc. All Rights Reserved.
 
-/*=============================================================================
-	AutomationDeviceClusterManager.cpp: Implements the FAutomationDeviceClusterManager class.
-=============================================================================*/
-
 #include "AutomationControllerPrivatePCH.h"
 
 void FAutomationDeviceClusterManager::Reset()
@@ -59,6 +55,7 @@ void FAutomationDeviceClusterManager::Remove(const FMessageAddress& MessageAddre
 	}
 }
 
+
 FString FAutomationDeviceClusterManager::GetGroupNameForDevice(const FDeviceState& DeviceState, const uint32 DeviceGroupFlags)
 {
 	FString OutGroupName;
@@ -112,6 +109,7 @@ FString FAutomationDeviceClusterManager::GetGroupNameForDevice(const FDeviceStat
 	return OutGroupName;
 }
 
+
 void FAutomationDeviceClusterManager::ReGroupDevices( const uint32 GroupFlags )
 {
 	//Get all the devices
@@ -152,6 +150,7 @@ void FAutomationDeviceClusterManager::ReGroupDevices( const uint32 GroupFlags )
 	}
 }
 
+
 int32 FAutomationDeviceClusterManager::GetNumClusters() const
 {
 	return Clusters.Num();
@@ -175,6 +174,7 @@ int32 FAutomationDeviceClusterManager::GetNumDevicesInCluster(const int32 Cluste
 	return Clusters[ClusterIndex].Devices.Num();
 }
 
+
 int32 FAutomationDeviceClusterManager::GetNumActiveDevicesInCluster(const int32 ClusterIndex) const
 {
 	check((ClusterIndex >= 0) && (ClusterIndex < Clusters.Num()));
@@ -189,18 +189,21 @@ int32 FAutomationDeviceClusterManager::GetNumActiveDevicesInCluster(const int32
 	return ActiveDevices;
 }
 
+
 FString FAutomationDeviceClusterManager::GetClusterGroupName(const int32 ClusterIndex) const
 {
 	check((ClusterIndex >= 0) && (ClusterIndex < Clusters.Num()));
 	return Clusters[ClusterIndex].ClusterName;
 }
 
+
 FString FAutomationDeviceClusterManager::GetClusterDeviceType(const int32 ClusterIndex) const
 {
 	check((ClusterIndex >= 0) && (ClusterIndex < Clusters.Num()));
 	return Clusters[ClusterIndex].DeviceTypeName;
 }
 
+
 FString FAutomationDeviceClusterManager::GetClusterDeviceName(const int32 ClusterIndex, const int32 DeviceIndex) const
 {
 	check((ClusterIndex >= 0) && (ClusterIndex < Clusters.Num()));
@@ -208,6 +211,7 @@ FString FAutomationDeviceClusterManager::GetClusterDeviceName(const int32 Cluste
 	return Clusters[ClusterIndex].Devices[DeviceIndex].GameInstanceName;
 }
 
+
 bool FAutomationDeviceClusterManager::FindDevice(const FMessageAddress& MessageAddress, int32& OutClusterIndex, int32& OutDeviceIndex)
 {
 	OutClusterIndex = INDEX_NONE;
@@ -321,4 +325,4 @@ bool FAutomationDeviceClusterManager::HasActiveDevice()
 		}
 	}
 	return IsDeviceAvailable;
-}
\ No newline at end of file
+}
diff --git a/Engine/Source/Developer/AutomationController/Private/AutomationDeviceClusterManager.h b/Engine/Source/Developer/AutomationController/Private/AutomationDeviceClusterManager.h
--- a/Engine/Source/Developer/AutomationController/Private/AutomationDeviceClusterManager.h
+++ b/Engine/Source/Developer/AutomationController/Private/AutomationDeviceClusterManager.h
@@ -1,9 +1,5 @@
 // Copyright 1998-2014 Epic Games, Inc. All Rights Reserved.
 
-/*=============================================================================
-	AutomationDeviceClusterManager.h: Declares the FAutomationDeviceClusterManager class.
-=============================================================================*/
-
 #pragma once
 
 
@@ -17,115 +13,117 @@ public:
 	void Reset();
 
 	/**
-	 * Add a new device from an AutomationWorkerFindWorkersResponse message, creating new clusters as needed
-	 * @param MessageAddress - The network address of an available worker
-	 * @param Message - The message that contains the device info
-	 * @param GroupFlags - Used find out which group to put the new device into
+	 * Add a new device from an AutomationWorkerFindWorkersResponse message, creating new clusters as needed.
+	 *
+	 * @param MessageAddress The network address of an available worker.
+	 * @param Message The message that contains the device info.
+	 * @param GroupFlags Used find out which group to put the new device into.
 	 */
 	void AddDeviceFromMessage(const FMessageAddress& MessageAddress, const FAutomationWorkerFindWorkersResponse& Message, const uint32 GroupFlags);
 
-	/** Remove a device (went offline, etc) */
+	/** Remove a device (went offline, etc). */
 	void Remove(const FMessageAddress& MessageAddress);
 
-	/** Returns number of unique device types */
+	/** Returns number of unique device types. */
 	int32 GetNumClusters() const;
 
-	/** Returns the sum of the device counts per cluster */
+	/** Returns the sum of the device counts per cluster. */
 	int32 GetTotalNumDevices() const;
 
-	/** Returns the number of devices in of a particular device type */
+	/** Returns the number of devices in of a particular device type. */
 	int32 GetNumDevicesInCluster (const int32 ClusterIndex) const;
 
-	/** Returns the number of active devices in of a particular device type */
+	/** Returns the number of active devices in of a particular device type. */
 	int32 GetNumActiveDevicesInCluster(const int32 ClusterIndex) const;
 
-	/** Returns the name of the cluster group */
+	/** Returns the name of the cluster group. */
 	FString GetClusterGroupName(const int32 ClusterIndex) const;
 
-	/** Returns the name of the devices within this cluster */
+	/** Returns the name of the devices within this cluster. */
 	FString GetClusterDeviceType (const int32 ClusterIndex) const;
 
-	/** Returns the name of a device within a cluster */
+	/** Returns the name of a device within a cluster. */
 	FString GetClusterDeviceName(const int32 ClusterIndex, const int32 DeviceIndex) const;
 
 	/** 
-	 * Finds the cluster/device index for a particular Guid
+	 * Finds the cluster/device index for a particular GUID.
 	 *
-	 * @param MessageAddress - Network address of the device
-	 * @param OutClusterIndex - the index of the platform this device was in
-	 * @param OutDeviceIndex - the index of the device within OutClusterIndex
-	 *
-	 * @return - true if the device was found within a cluster
+	 * @param MessageAddress Network address of the device.
+	 * @param OutClusterIndex The index of the platform this device was in.
+	 * @param OutDeviceIndex The index of the device within OutClusterIndex.
+	 * @return true if the device was found within a cluster.
 	 */
 	bool FindDevice(const FMessageAddress& MessageAddress, int32& OutClusterIndex, int32& OutDeviceIndex);
 
 	/** 
-	 * Returns the message address of the device specified 
-	 * @param ClusterIndex	- index of the device cluster
-	 * @param DeviceIndex	- index of the device within the cluster
+	 * Returns the message address of the device specified.
+	 *
+	 * @param ClusterIndex Index of the device cluster.
+	 * @param DeviceIndex Index of the device within the cluster.
 	 *
 	 * @return The message address of this device
 	 */
 	FMessageAddress GetDeviceMessageAddress(const int32 ClusterIndex, const int32 DeviceIndex) const;
 
 	/**
-	 * Returns the devices that have been reserved for a particular test
+	 * Returns the devices that have been reserved for a particular test.
 	 *
-	 * @param ClusterIndex - the device type we are testing
-	 * @param Report - The automation test that we are trying to run (when we have enough participants)
-	 *
-	 * @return - The array of devices that are currently signed up for running this test
+	 * @param ClusterIndex The device type we are testing.
+	 * @param Report The automation test that we are trying to run (when we have enough participants).
+	 * @return The array of devices that are currently signed up for running this test.
 	 */
 	TArray<FMessageAddress> GetDevicesReservedForTest(const int32 ClusterIndex, TSharedPtr <IAutomationReport> Report);
 
 	/**
-	 * Returns the current test this device should be running
-	 * @param ClusterIndex	- index of the device cluster
-	 * @param DeviceIndex	- index of the device within the cluster
-	 * @return - NULL if this device is available or the test structure if it is in-process
+	 * Returns the current test this device should be running.
+	 *
+	 * @param ClusterIndex Index of the device cluster.
+	 * @param DeviceIndex Index of the device within the cluster.
+	 * @return nullptr if this device is available or the test structure if it is in-process.
 	 */
 	TSharedPtr <IAutomationReport> GetTest(const int32 ClusterIndex, const int32 DeviceIndex) const;
 
 	/**
-	 * Sets the current test being run by the particular device
+	 * Sets the current test being run by the particular device.
 	 *
-	 * @param ClusterIndex	- The Cluster which holds the device we intend to run the test on
-	 * @param DeviceIndex	- The devices index in the cluster where we intend to run the test on
-	 * @param NewReport		- The test we are going to run
+	 * @param ClusterIndex The Cluster which holds the device we intend to run the test on.
+	 * @param DeviceIndex The devices index in the cluster where we intend to run the test on.
+	 * @param NewReport The test we are going to run.
 	 */
 	void SetTest(const int32 ClusterIndex, const int32 DeviceIndex, TSharedPtr <IAutomationReport> NewReport);
 
 	/**
-	 * Resets all cases on devices which are running the specified test to NULL
+	 * Resets all cases on devices which are running the specified test to NULL.
 	 *
-	 * @param ClusterIndex	- The Cluster where we want all devices running the test to cease
-	 * @param InTest		- The test we are stopping
+	 * @param ClusterIndex The Cluster where we want all devices running the test to cease.
+	 * @param InTest The test we are stopping.
 	 */
 	void ResetAllDevicesRunningTest( const int32 ClusterIndex, IAutomationReportPtr InTest );
 
 	/**
-	 * Disable a device as it is no longer available - keep around to get the results
-	 * @param ClusterIndex - the index of the platform this device was in
-	 * @param DeviceIndex - the index of the device within OutClusterIndex
+	 * Disable a device as it is no longer available - keep around to get the results.
+	 *
+	 * @param ClusterIndex The index of the platform this device was in.
+	 * @param DeviceIndex The index of the device within OutClusterIndex.
 	 */
 	void DisableDevice( const int32 ClusterIndex, const int32 DeviceIndex );
 
 	/**
-	 * Check if a device is enabled
-	 * @param ClusterIndex - the index of the platform this device was in
-	 * @param DeviceIndex - the index of the device within OutClusterIndex
+	 * Check if a device is enabled.
+	 *
+	 * @param ClusterIndex The index of the platform this device was in.
+	 * @param DeviceIndex The index of the device within OutClusterIndex.
 	 */
 	bool DeviceEnabled( const int32 ClusterIndex, const int32 DeviceIndex );
 
 	/**
-	 * Check if there are any active game instances left
-	 * @return True if there are some game instances
+	 * Check if there are any active game instances left.
+	 *
+	 * @return True if there are some game instances.
 	 */
 	bool HasActiveDevice();
 
-	/**
-	 * Regroups the device clusters based off the current group flags
-	 */
+	/** Regroups the device clusters based off the current group flags. */
 	void ReGroupDevices( const uint32 GroupFlags );
 
 private:
@@ -209,5 +207,3 @@ private:
 	/** Array of all clusters */
 	TArray <FDeviceCluster> Clusters;
 };
-
-
diff --git a/Engine/Source/Developer/AutomationController/Private/AutomationReport.cpp b/Engine/Source/Developer/AutomationController/Private/AutomationReport.cpp
--- a/Engine/Source/Developer/AutomationController/Private/AutomationReport.cpp
+++ b/Engine/Source/Developer/AutomationController/Private/AutomationReport.cpp
@@ -1,9 +1,5 @@
 // Copyright 1998-2014 Epic Games, Inc. All Rights Reserved.
 
-/*=============================================================================
-	AutomationTestStatus.cpp: Implements the AutomationTestStatus class.
-=============================================================================*/
-
 #include "AutomationControllerPrivatePCH.h"
 
 
@@ -86,6 +82,7 @@ int32 FAutomationReport::GetTotalNumChildren() const
 	return Total;
 }
 
+
 void FAutomationReport::GetEnabledTestNames(TArray<FString>& OutEnabledTestNames, FString CurrentPath) const
 {
 	//if this is a leaf and this test is enabled
@@ -110,6 +107,7 @@ void FAutomationReport::GetEnabledTestNames(TArray<FString>& OutEnabledTestNames
 	return;
 }
 
+
 void FAutomationReport::SetEnabledTests(const TArray<FString>& EnabledTests, FString CurrentPath)
 {
 	if (ChildReports.Num() == 0)
@@ -140,7 +138,7 @@ void FAutomationReport::SetEnabledTests(const TArray<FString>& EnabledTests, FSt
 	}
 }
 
-/** Recursively gets the number of enabled tests */
+
 int32 FAutomationReport::GetEnabledTestsNum() const
 {
 	int32 Total = 0;
diff --git a/Engine/Source/Developer/AutomationController/Private/AutomationReport.h b/Engine/Source/Developer/AutomationController/Private/AutomationReport.h
--- a/Engine/Source/Developer/AutomationController/Private/AutomationReport.h
+++ b/Engine/Source/Developer/AutomationController/Private/AutomationReport.h
@@ -1,266 +1,75 @@
 // Copyright 1998-2014 Epic Games, Inc. All Rights Reserved.
 
-/*=============================================================================
-	AutomationTestResults.h: Declares the AutomationTestResults interface.
-=============================================================================*/
-
 #pragma once
 
 
 /**
  * Interface for automation test results modules.
  */
-class FAutomationReport : public IAutomationReport
+class FAutomationReport
+	: public IAutomationReport
 {
 public:
 
 	FAutomationReport(FAutomationTestInfo& TestInfo, bool bIsParent = false);
 
-	// Begin IAutomationReport Interface
+public:
 
-	/** Remove all child tests */
-	virtual void Empty() override;
+	// IAutomationReport Interface
 
-	/** Get the asset associated with tis report */
+	virtual void Empty() override;
 	virtual FString GetAssetName() const override;
-
-	/**
-	 * Returns the complete command line for an automation test including any relevant parameters
-	 */
 	virtual FString GetCommand() const override;
-
-	/**
-	 * Returns the name of this level in the test hierarchy for the purposes of grouping.
-	 * (Editor.Maps.LoadAll.parameters) would have 3 internal tree nodes and a variable number of leaf nodes depending on how many maps existed
-	 * @param bAddDecoration - True is the name should have the number of child tests appended
-	 * @return the name of this level in the hierarchy for use in UI
-	 */
 	virtual const FString& GetDisplayName() const override;
-
-	/**
-	 * Returns the name of this level in the test hierarchy for the purposes of UI.
-	 * (Editor.Maps.LoadAll.parameters) would have 3 internal tree nodes and a variable number of leaf nodes depending on how many maps existed
-	 * @return the name of this level in the hierarchy for use in UI
-	 */
 	virtual FString GetDisplayNameWithDecoration() const override;
-
-	/** Recursively gets the number of child nodes */
 	virtual int32 GetTotalNumChildren() const override;
-	/** Recursively gets the number of enabled tests */
 	virtual int32 GetEnabledTestsNum() const override;
-	/** Recursively gets the names of enabled tests */
 	virtual void GetEnabledTestNames(TArray<FString>& OutEnabledTestNames, FString CurrentPath) const override;
-	/** Recursively sets the enabled flag based on if the current test exists in the enabled tests array */
 	virtual void SetEnabledTests(const TArray<FString>& EnabledTests, FString CurrentPath) override;
-
-	/** Return if this test should be executed */
 	virtual bool IsEnabled() const override;
-	/** Sets whether this test should be executed or not */
 	virtual void SetEnabled(bool bShouldBeEnabled) override;
-	
-	/** Sets whether this test is supported on a particular platform */
 	virtual void SetSupport(const int32 ClusterIndex) override;
-	/** Returns if a particular platform is supported */
 	virtual bool IsSupported(const int32 ClusterIndex) const override;
-
-	/** Sets the test type for parent reports */
 	virtual void SetTestType( const uint8 TestType ) override;
-	/** Returns the test type */
 	virtual uint8 GetTestType( ) const override;
-
-	/** Returns true if this is a parent test report */
 	virtual const bool IsParent( ) override;
-	/** Returns true if this is a smoke test report */
 	virtual const bool IsSmokeTest( ) override;
-	/**
-	 * Filters the visible tests based on search text, execution status, regression test, etc
-	 * @param InFilterDesc - The description of how to decide if this report should be visible
-	 * @param ParentPassedFilter - If the parent passed the filter
-	 * @return whether this report or any of its children passed the filter
-	 */
 	virtual bool SetFilter( TSharedPtr< AutomationFilterCollection > InFilter, const bool ParentPassedFilter = false ) override;
-
-	/** Returns the array of child reports that should be visible to the UI based on filtering */
 	virtual TArray<TSharedPtr<IAutomationReport> >& GetFilteredChildren() override;
-
-	/** Returns the array of child reports */
 	virtual TArray<TSharedPtr<IAutomationReport> >& GetChildReports() override;
-
-	/** Updates the report when the number of clusters changes */
 	virtual void ClustersUpdated(const int32 NumClusters) override;
-
-	/** 
-	 * Recursively resets the report to "needs to be run", clears cached warnings and errors
-	 * @param NumTestPasses - The number of test passes so we know how many results to create
-	 */
 	virtual void ResetForExecution(const int32 NumTestPasses) override;
-
-	/**
-	 * Sets the results of the test for use by the UI
-	 * @param ClusterIndex - Index of the platform reporting the results of this test.  See AutomationDeviceClusterManager
-	 * @param PassIndex - Which test pass these results are for
-	 * @param InResults - The new set of results
-	 */
 	virtual void SetResults( const int32 ClusterIndex, const int32 PassIndex, const FAutomationTestResults& InResults ) override;
-
-	/**
-	 * Returns completion statistics for this branch of the testing hierarchy
-	 * @param ClusterIndex - Index of the platform reporting the results of this test.  See AutomationDeviceClusterManager
-	 * @param PassIndex - Which test pass to get the status of
-	 * @param OutCompletionState - Collection structure for execution statistics
-	 */
 	virtual void GetCompletionStatus(const int32 ClusterIndex, const int32 PassIndex, FAutomationCompleteState& OutCompletionState) override;
-
-	/**
-	 * Returns the state of the test (not run, in process, success, failure)
-	 * @param ClusterIndex - Index of the platform reporting the results of this test.  See AutomationDeviceClusterManager
-	 * @param PassIndex - Which test pass to get the state of
-	 * @return the current state of the test
-	 */
 	virtual EAutomationState::Type GetState(const int32 ClusterIndex, const int32 PassIndex) const override;
-
-	/**
-	 * Gets a copy of errors and warnings that were found
-	 *
-	 * @param ClusterIndex - Index of the platform we are requesting test results for.
-	 * @param PassIndex - Index of the test pass to get the results for.
-	 *
-	 * @return The collection of results for the given cluster index
-	 */
 	virtual const FAutomationTestResults& GetResults( const int32 ClusterIndex, const int32 PassIndex ) override;
-
-	/**
-	 * Gets the number of available test results for a given cluster
-	 * @param ClusterIndex - Index of the platform .
-	 * @return The number of results available for this cluster
-	 */
 	virtual const int32 GetNumResults( const int32 ClusterIndex ) override;
-
-	/**
-	 * Finds the current pass by looking at the current state
-	 * @param ClusterIndex - Index of the platform.
-	 * @return The current pass index
-	 */
 	virtual const int32 GetCurrentPassIndex( const int32 ClusterIndex ) override;
-
-	/**
-	 * Gets the name of the instance that ran the test
-	 * @param ClusterIndex - Index of the platform reporting the results of this test.  See AutomationDeviceClusterManager
-	 * @return the name of the device
-	 */
 	virtual FString GetGameInstanceName( const int32 ClusterIndex ) override;
-
-	/**
-	 * Add a child test to the hierarchy, creating internal tree nodes as needed
-	 * If NewTestName is Editor.Maps.Loadall.TestName, this will create nodes for Editor, Maps, Loadall, and then a leaf node for the testname with the associated commandline
-	 * @param TestInfo - Structure containing all the test info
-	 * @param ClusterIndex - Index of the platform reporting the results of this test.  See AutomationDeviceClusterManager
-	 * @param NumPasses - The number of passes we are going to perform.  Used to make sure we have enough results.
-	 * @return - the automation report
-	 */
 	virtual TSharedPtr<IAutomationReport> EnsureReportExists(FAutomationTestInfo& TestInfo, const int32 ClusterIndex, const int32 NumPasses) override;
-
-	/**
-	 * Returns the next test in the hierarchy to run
-	 * @param bOutAllTestsComplete - Whether or not all enabled tests have completed execution for this platform (cluster index)
-	 * @param ClusterIndex - Index of the platform reporting the results of this test.  See AutomationDeviceClusterManager
-	 * @param PassIndex - Which test pass we are currently on
-	 * @param NumDevicesInCluster - The number of devices which are in this cluster
-	 *
-	 * @return The next report
-	 */
 	virtual TSharedPtr<IAutomationReport> GetNextReportToExecute(bool& bOutAllTestsComplete, const int32 ClusterIndex, const int32 PassIndex, const int32 NumDevicesInCluster) override;
-
-	/** 
-	 * Have any results got errors
-	 * @return true if they have errors
-	 */
 	virtual const bool HasErrors() override;
-
-	/** 
-	 * Have any results got warnings
-	 * @return true if they have warnings
-	 */
 	virtual const bool HasWarnings() override;
-
-	/**
-	 * Gets the min and max time this test took to execute
-	 * @param OutMinTime - Minimum execution time for all device types
-	 * @param OutMaxTime - Maximum execution time for all device types
-	 * @return whether any test has completed successfully
-	 */
 	virtual const bool GetDurationRange(float& OutMinTime, float& OutMaxTime) override;
-
-	/**
-	 * Get the number of devices which have been given this test to run
-	 *
-	 * @return - The number of devices who have been given this test to run
-	 */
 	virtual const int32 GetNumDevicesRunningTest() const override;
-
-	/**
-	 * Get the number of participants this test requires.
-	 *
-	 * @Return - The number of devices needed for this test to execute
-	 */
 	virtual const int32 GetNumParticipantsRequired() const override;
-	
-	/**
-	 * Set the number of participants this test requires if less than what is already set.
-	 *
-	 * @param NewCount - The number of further devices 
-	 */
 	virtual void SetNumParticipantsRequired( const int32 NewCount ) override;
-
-	/**
-	 * Increment the number of network responses
-	 * @return - whether we have received responses from all participants
-	 */
 	virtual bool IncrementNetworkCommandResponses() override;
-
-	/**
-	 * Resets the number of network responses back to zero
-	 */
 	virtual void ResetNetworkCommandResponses() override;
-
-	/**
-	 * Should we expand this node in the UI - A child has passed the filter
-	 * @return true if we should expand the node
-	 */
 	virtual const bool ExpandInUI() const override;
-
-	/** Stop the test which is creating this report */
 	virtual void StopRunningTest() override;
-
-	/**
-	* Notification on whether we should, or should not, track this reports history
-	*/
 	virtual void TrackHistory(const bool bShouldTrack, const int32 NumReportsToTrack) override;
-
-	/**
-	* Get the history items of this particular test
-	*
-	* @return A reference to the items of this tests previous runs.
-	*/
 	virtual const TArray<TSharedPtr<FAutomationHistoryItem>>& GetHistory() const override;
-
-	// End IAutomationReport Interface
 	
 private:
 
-	/** 
-	 * Export the current report as part of it's tracked history 
-	 */
+	/** Export the current report as part of it's tracked history. */
 	void AddToHistory();
 
-	/** 
-	 * Load this reports tracked history 
-	 */
+	/** Load this reports tracked history. */
 	void LoadHistory();
 
-	/**
-	 * Update what is tracked for this reports history.
-	 */
+	/** Update what is tracked for this reports history. */
 	void MaintainHistory();
 
 private:
diff --git a/Engine/Source/Developer/AutomationController/Private/AutomationReportManager.cpp b/Engine/Source/Developer/AutomationController/Private/AutomationReportManager.cpp
--- a/Engine/Source/Developer/AutomationController/Private/AutomationReportManager.cpp
+++ b/Engine/Source/Developer/AutomationController/Private/AutomationReportManager.cpp
@@ -1,9 +1,5 @@
 // Copyright 1998-2014 Epic Games, Inc. All Rights Reserved.
 
-/*=============================================================================
-	AutomationReportManager.cpp: Implements the FAutomationReportManager class.
-=============================================================================*/
-
 #include "AutomationControllerPrivatePCH.h"
 
 
@@ -21,11 +17,13 @@ void FAutomationReportManager::Empty()
 	ReportRoot->Empty();
 }
 
+
 void FAutomationReportManager::ClustersUpdated( const int32 NumClusters )
 {
 	ReportRoot->ClustersUpdated(NumClusters);
 }
 
+
 void FAutomationReportManager::ResetForExecution(const int32 NumTestPasses)
 {
 	//Recursively prepare all tests for execution
@@ -38,6 +36,7 @@ void FAutomationReportManager::StopRunningTests()
 	ReportRoot->StopRunningTest();
 }
 
+
 TSharedPtr<IAutomationReport> FAutomationReportManager::GetNextReportToExecute(bool& bOutAllTestsComplete, const int32 ClusterIndex, const int32 PassIndex, const int32 NumDevicesInCluster)
 {
 	TSharedPtr<IAutomationReport> ReportToExecute = ReportRoot->GetNextReportToExecute(bOutAllTestsComplete, ClusterIndex, PassIndex, NumDevicesInCluster);
@@ -75,16 +74,19 @@ int32 FAutomationReportManager::GetEnabledTestsNum() const
 	return ReportRoot->GetEnabledTestsNum();
 }
 
+
 void FAutomationReportManager::GetEnabledTestNames(TArray<FString>& OutEnabledTestNames) const
 {
 	ReportRoot->GetEnabledTestNames(OutEnabledTestNames,FString());
 }
 
+
 void FAutomationReportManager::SetEnabledTests(const TArray<FString>& EnabledTests)
 {
 	ReportRoot->SetEnabledTests(EnabledTests,FString());
 }
 
+
 const bool FAutomationReportManager::ExportReport(uint32 FileExportTypeMask, const int32 NumDeviceClusters)
 {
 	// The results log. Errors, warnings etc are added to this
diff --git a/Engine/Source/Developer/AutomationController/Private/AutomationReportManager.h b/Engine/Source/Developer/AutomationController/Private/AutomationReportManager.h
--- a/Engine/Source/Developer/AutomationController/Private/AutomationReportManager.h
+++ b/Engine/Source/Developer/AutomationController/Private/AutomationReportManager.h
@@ -1,9 +1,5 @@
 // Copyright 1998-2014 Epic Games, Inc. All Rights Reserved.
 
-/*=============================================================================
-	AutomationReportManager.h: Declares the FAutomationReportManager class.
-=============================================================================*/
-
 #pragma once
 
 
@@ -24,7 +20,7 @@ public:
 
 	/** 
 	 * Resets for the next run of tests 
-	 * @param NumTestPasses - The number of test results to store
+	 * @param NumTestPasses The number of test results to store
 	 */
 	void ResetForExecution(const int32 NumTestPasses);
 
@@ -32,80 +28,87 @@ public:
 	void StopRunningTests();
 
 	/**
-	 * Searches for the next test to execute 
-	 * @param bOutAllTestsComplete - Returns True if all tests for this device cluster have been completed
-	 * @param ClusterIndex - Index of the platform reporting the results of this test.  See AutomationDeviceClusterManager
-	 * @param PassIndex - Index of the current test pass we are on
-	 * @param NumDevicesInCluster - The number of devices which are available in this cluster
+	 * Searches for the next test to execute.
+	 *
+	 * @param bOutAllTestsComplete Returns True if all tests for this device cluster have been completed.
+	 * @param ClusterIndex Index of the platform reporting the results of this test.  See AutomationDeviceClusterManager.
+	 * @param PassIndex Index of the current test pass we are on.
+	 * @param NumDevicesInCluster The number of devices which are available in this cluster.
 	 */
 	TSharedPtr<IAutomationReport> GetNextReportToExecute(bool& bOutAllTestsComplete, const int32 ClusterIndex, const int32 PassIndex, const int32 NumDevicesInCluster);
 
 	/**
-	 * Ensures the nested tree exists ("blueprint.test.all") would have 3 levels other than the root
-	 * @param TestInfo - Structure containing the test info
-	 * @param ClusterIndex - Index of the platform reporting the results of this test.  See AutomationDeviceClusterManager
-	 * @param NumPasses - The number of test passes to create reports for
-	 * @return The automation report
+	 * Ensures the nested tree exists ("blueprint.test.all") would have 3 levels other than the root.
+	 *
+	 * @param TestInfo Structure containing the test info.
+	 * @param ClusterIndex Index of the platform reporting the results of this test.  See AutomationDeviceClusterManager.
+	 * @param NumPasses The number of test passes to create reports for.
+	 * @return The automation report.
 	 */
 	TSharedPtr<IAutomationReport> EnsureReportExists(FAutomationTestInfo& TestInfo, const int32 ClusterIndex, const int32 NumPasses);
 
-	/** Filters the visible tests based on name, status, speed */
+	/** Filters the visible tests based on name, status, speed. */
 	void SetFilter( TSharedPtr< AutomationFilterCollection > InFilter );
 	
-	/** gets array of filtered tests to display in the UI*/
+	/** gets array of filtered tests to display in the UI. */
 	TArray <TSharedPtr <IAutomationReport> >& GetFilteredReports();
 
-	/** Sets whether all visible tests are enabled or not */
+	/** Sets whether all visible tests are enabled or not. */
 	void SetVisibleTestsEnabled(const bool bEnabled);
-	/** Returns number of tests that will be run */
+
+	/** Returns number of tests that will be run. */
 	int32 GetEnabledTestsNum () const;
 
-	/** Gets the names of enabled tests */
+	/** Gets the names of enabled tests. */
 	void GetEnabledTestNames(TArray<FString>& OutEnabledTestNames) const;
-	/** Sets the enabled tests based off the passed list of enabled tests */
+
+	/** Sets the enabled tests based off the passed list of enabled tests. */
 	void SetEnabledTests(const TArray<FString>& EnabledTests);
 
-	/** Sets the current test pass */
+	/** Sets the current test pass. */
 	void SetCurrentTestPass(int32 PassIndex){ CurrentTestPass = PassIndex; }
 
 	/**
-	 * Export the automation report
-	 * @param FileExportTypeMask - The file export mask - errors, warning etc
-	 * @param NumDeviceClusters - The number of devices in a cluster
+	 * Export the automation report.
+	 *
+	 * @param FileExportTypeMask The file export mask - errors, warning etc.
+	 * @param NumDeviceClusters The number of devices in a cluster.
 	 */
 	const bool ExportReport( uint32 FileExportTypeMask, const int32 NumDeviceClusters );
 
 	/**
-	 * Notification on whether we should, or should not, track report histories
+	 * Notification on whether we should, or should not, track report histories.
 	 */
 	void TrackHistory(const bool bShouldTrack, const int32 NumReportsToTrack);
 
 protected:
 
 	/**
-	 * Used in report generation to find leaf reports
-	 * @param InReport - The report to add
- 	 * @param NumDeviceClusters - The number of devices in a cluster
-	 * @param ResultsLog - Message log to store the test results
-	 * @param FileExportTypeMask - The file export mask - errors, warning etc
+	 * Used in report generation to find leaf reports.
+	 *
+	 * @param InReport The report to add.
+ 	 * @param NumDeviceClusters - The number of devices in a cluster.
+	 * @param ResultsLog - Message log to store the test results.
+	 * @param FileExportTypeMask - The file export mask - errors, warning etc.
 	 */
 	void AddResultReport( TSharedPtr< IAutomationReport > InReport, const int32 NumDeviceClusters, TArray< FString >& ResultsLog, uint32 FileExportTypeMask );
 
 	/**
-	 * Used in report generation to find leaf reports
-	 * @param InReport - The parent report
-	 * @param NumDeviceClusters - The number of devices in a cluster
-	 * @param ResultsLog - Message log to store the test results
-	 * @param FileExportTypeMask - The file export mask - errors, warning etc
+	 * Used in report generation to find leaf reports.
+	 *
+	 * @param InReport The parent report.
+	 * @param NumDeviceClusters The number of devices in a cluster.
+	 * @param ResultsLog Message log to store the test results.
+	 * @param FileExportTypeMask The file export mask - errors, warning etc.
 	 */
 	void FindLeafReport( TSharedPtr< IAutomationReport > InReport, const int32 NumDeviceClusters, TArray< FString >& ResultsLog, uint32 FileExportTypeMask );
 
 
 private:
-	/** Root node of the hierarchy, just there to ensure little code duplication */
+
+	/** Root node of the hierarchy, just there to ensure little code duplication. */
 	TSharedPtr <IAutomationReport> ReportRoot;
-	/** Which test pass we are currently on */
+
+	/** Which test pass we are currently on. */
 	int32 CurrentTestPass;
 };
-
-
diff --git a/Engine/Source/Developer/AutomationController/Public/AutomationController.h b/Engine/Source/Developer/AutomationController/Public/AutomationController.h
--- a/Engine/Source/Developer/AutomationController/Public/AutomationController.h
+++ b/Engine/Source/Developer/AutomationController/Public/AutomationController.h
@@ -1,9 +1,5 @@
 // Copyright 1998-2014 Epic Games, Inc. All Rights Reserved.
 
-/*=============================================================================
-	AutomationController.h: Automation Controller module public header file.
-=============================================================================*/
-
 #pragma once
 
 
diff --git a/Engine/Source/Developer/AutomationController/Public/Interfaces/IAutomationControllerManager.h b/Engine/Source/Developer/AutomationController/Public/Interfaces/IAutomationControllerManager.h
--- a/Engine/Source/Developer/AutomationController/Public/Interfaces/IAutomationControllerManager.h
+++ b/Engine/Source/Developer/AutomationController/Public/Interfaces/IAutomationControllerManager.h
@@ -1,21 +1,15 @@
 // Copyright 1998-2014 Epic Games, Inc. All Rights Reserved.
 
-/*=============================================================================
-	IAutomationController.h: Declares the IAutomationController interface.
-=============================================================================*/
-
 #pragma once
 
-/**
- * Type definition for shared pointers to instances of IAutomationControllerManager.
- */
+
+/** Type definition for shared pointers to instances of IAutomationControllerManager. */
 typedef TSharedPtr<class IAutomationControllerManager> IAutomationControllerManagerPtr;
 
-/**
- * Type definition for shared references to instances of IAutomationControllerManager.
- */
+/** Type definition for shared references to instances of IAutomationControllerManager. */
 typedef TSharedRef<class IAutomationControllerManager> IAutomationControllerManagerRef;
 
+
 // Enum to indicate the state of the workers in the automation controller
 namespace EAutomationControllerModuleState 
 {
@@ -30,6 +24,7 @@ namespace EAutomationControllerModuleState
 	};
 };
 
+
 /**
 * Enum used to set which reports to export
 */
@@ -86,6 +81,7 @@ namespace EFileExportType
 
 };
 
+
 namespace EAutomationDeviceGroupTypes
 {
 	enum Type
@@ -156,72 +152,73 @@ DECLARE_DELEGATE(FOnAutomationControllerTestsComplete);
 class IAutomationControllerManager
 {
 public:
-	virtual ~IAutomationControllerManager() { }
 
 	/** 
-	 * Broadcast a request for workers to ping if they are available 
+	 * Broadcast a request for workers to ping if they are available.
 	 * 
-	 * @param InSessionID the session ID to run the test on
+	 * @param InSessionID the session ID to run the test on.
 	 */
 	virtual void RequestAvailableWorkers( const FGuid& InSessionId ) = 0;
 
-	/** Send request for tests that are available to run */
+	/** Send request for tests that are available to run. */
 	virtual void RequestTests() = 0;
 
 	/** 
-	 * Resets all enabled tests to be able to run during Tick by local/remote machines
+	 * Resets all enabled tests to be able to run during Tick by local/remote machines.
 	 *
-	 * @param bIsLocalSeesion - indicate if this is a local session. We do not time out local session automation tests
+	 * @param bIsLocalSeesion Indicate if this is a local session. We do not time out local session automation tests.
 	 */
 	virtual void RunTests( const bool bIsLocalSession = false ) = 0;
 
-	/**
-	 * Stops all running tests.
-	 */
+	/** Stops all running tests. */
 	virtual void StopTests() = 0;
 
-	/** Initializes the automation controller */
+	/** Initializes the automation controller. */
 	virtual void Init() = 0;
 
 	/**
-	 * Request an asset to be loaded in the editor
+	 * Request an asset to be loaded in the editor.
 	 *
-	 * @param InAssetName - the asset name
+	 * @param InAssetName The asset name.
 	 */
 	virtual void RequestLoadAsset( const FString& InAssetName ) = 0;
 
-	/** Tick function that will execute enabled tests for different device clusters */
+	/** Tick function that will execute enabled tests for different device clusters. */
 	virtual void Tick() = 0;
 
 	/**
-	 * Sets the number of test passes
+	 * Sets the number of test passes.
 	 *
-	 * @param InNumPasses - the number of test passes to perform
+	 * @param InNumPasses The number of test passes to perform.
 	 */
 	virtual void SetNumPasses(const int32 InNumPasses) = 0;
 
 	/**
-	 * Get the number of test passes
+	 * Get the number of test passes.
+	 *
+	 * @return Number of passes.
 	 */
 	virtual int32 GetNumPasses() = 0;
 
 	/**
-	 * Returns if we are using full size screen shots
+	 * Returns if we are using full size screen shots.
+	 *
+	 * @return true if using full size screenshots, false otherwise.
 	 */
 	virtual bool IsUsingFullSizeScreenshots() const = 0;
 
 	/**
-	 * Sets if we are using full size screen shots
+	 * Sets if we are using full size screen shots.
 	 */
 	virtual void SetUsingFullSizeScreenshots(const  bool bNewValue ) = 0;
 
 	/**
-	 * Returns if screenshots are allowed
+	 * Returns if screenshots are allowed.
 	 */
 	virtual bool IsScreenshotAllowed() const = 0;
 
 	/**
-	 * Sets if screenshots are enabled
+	 * Sets if screenshots are enabled.
 	 */
 	virtual void SetScreenshotsEnabled( const bool bNewValue ) = 0;
 
@@ -256,10 +253,10 @@ public:
 	virtual FString GetDeviceTypeName(const int32 ClusterIndex) const = 0;
 
 	/**
-	 * Get a game instance name
+	 * Get a game instance name.
 	 *
-	 * @param ClusterIndex The cluster Index
- 	 * @param DeviceIndex The Device Index
+	 * @param ClusterIndex The cluster Index.
+ 	 * @param DeviceIndex The Device Index.
 	 */
 	virtual FString GetGameInstanceName(const int32 ClusterIndex, const int32 DeviceIndex) const = 0;
 
@@ -274,16 +271,16 @@ public:
 	virtual int32 GetEnabledTestsNum() const = 0;
 
 	/**
-	 * Gets the names of all the enabled tests
+	 * Gets the names of all the enabled tests.
 	 *
-	 * @param OutEnabledTestNames The list to populate with enabled test names
+	 * @param OutEnabledTestNames The list to populate with enabled test names.
 	 */
 	virtual void GetEnabledTestNames(TArray<FString>& OutEnabledTestNames) const = 0;
 
 	/**
-	 * Sets any tests that match a name in the enabled tests array
+	 * Sets any tests that match a name in the enabled tests array.
 	 *
-	 * @param EnabledTests An array of test names that will be enabled
+	 * @param EnabledTests An array of test names that will be enabled.
 	 */
 	virtual void SetEnabledTests(const TArray<FString>& EnabledTests) = 0;
 
@@ -317,82 +314,69 @@ public:
 	virtual bool IsVisualCommandletFilterOn(void) const = 0;
 
 	/**
-	 * Check if the automation tests have completed
+	 * Check if the automation tests have completed.
 	 *
-	 * @return true if the tests are available
+	 * @return true if the tests are available.
 	 */
 	virtual const bool CheckTestResultsAvailable() const = 0;
 
 	/**
-	 * Check if the automation tests results have errors
+	 * Check if the automation tests results have errors.
 	 *
-	 * @return true if the tests have errors
+	 * @return true if the tests have errors.
 	 */
 	virtual const bool ReportsHaveErrors() const = 0;
 
 	/**
-	 * Check if the automation tests results have warnings
+	 * Check if the automation tests results have warnings.
 	 *
-	 * @return true if the tests have warnings
+	 * @return true if the tests have warnings.
 	 */
 	virtual const bool ReportsHaveWarnings() const = 0;
 
 	/**
-	 * Check if the automation tests results have logs
+	 * Check if the automation tests results have logs.
 	 *
-	 * @return true if the tests have logs
+	 * @return true if the tests have logs.
 	 */
 	virtual const bool ReportsHaveLogs() const = 0;
 
 	/**
-	 * Remove results from the automation controller module
+	 * Remove results from the automation controller module.
 	 */
 	virtual void ClearAutomationReports() = 0;
 
 	/**
-	 * Generate an automation report
+	 * Generate an automation report.
 	 *
-	 * @param FileExporTypeMask - The types of report to export. Warning errors etc
+	 * @param FileExporTypeMask The types of report to export. Warning errors etc.
 	 */
 	virtual const bool ExportReport( uint32 FileExportTypeMask ) = 0;
 
 	/** 
-	 * Check that the test we are looking to run is runnable
+	 * Check that the test we are looking to run is runnable.
 	 *
-	 * @param InReport - The test we are checking is runnable
-	 *
-	 * @return true if the test can be run, FALSE if not
+	 * @param InReport The test we are checking is runnable.
+	 * @return true if the test can be run, false if not.
 	 */
 	virtual bool IsTestRunnable( IAutomationReportPtr InReport ) const = 0;
 
-	/**
-	 * Removes all callbacks.
-	 */
+	/** Removes all callbacks. */
 	virtual void RemoveCallbacks() = 0;
 
-	/**
-	 * Shuts down the messaging system.
-	 */
+	/** Shuts down the manager. */
 	virtual void Shutdown() = 0;
 
-	/**
-	 * Starts up the messaging system.
-	 */
+	/** Starts up the manager. */
 	virtual void Startup() = 0;
 
-	/**
-	 * Checks if a device group flag is set.
-	 */
+	/** Checks if a device group flag is set. */
 	virtual bool IsDeviceGroupFlagSet( EAutomationDeviceGroupTypes::Type InDeviceGroup ) const = 0;
 
-	/**
-	 * Toggles a device group flag.
-	 */
+	/** Toggles a device group flag. */
 	virtual void ToggleDeviceGroupFlag( EAutomationDeviceGroupTypes::Type InDeviceGroup ) = 0;
 
-	/**
-	 * Updates the clusters when the device grouping changes.
-	 */
+	/** Updates the clusters when the device grouping changes. */
 	virtual void UpdateDeviceGroups( ) = 0;
 
 	/**
@@ -401,21 +385,21 @@ public:
 	virtual void SetTestsCompleteCallback(const FOnAutomationControllerTestsComplete& NewCallback) = 0;
 
 	/**
-	 * Dictate whether to save the reports history and the number of history items to track
+	 * Dictate whether to save the reports history and the number of history items to track.
 	 *
-	 * @param bShouldTrack - Flag that determines whether to track history
-	 * @param NumReportsToTrack - The number of history reports to keep
+	 * @param bShouldTrack Flag that determines whether to track history.
+	 * @param NumReportsToTrack The number of history reports to keep.
 	 */
 	virtual void TrackReportHistory(const bool bShouldTrack, const int32 NumReportsToTrack) = 0;
 
 	/**
-	* Returns whether the controller is tracking history of reports.
-	*/
+	 * Returns whether the controller is tracking history of reports.
+	 */
 	virtual const bool IsTrackingHistory() const = 0;
 
 	/**
-	* Returns the number of history items the controller is maintaining.
-	*/
+	 * Returns the number of history items the controller is maintaining.
+	 */
 	virtual const int32 GetNumberHistoryItemsTracking() const = 0;
 
 public:
@@ -440,4 +424,9 @@ public:
 	 * @return The delegate.
 	 */
 	virtual FOnAutomationControllerTestsRefreshed& OnTestsRefreshed( ) = 0;
+
+public:
+
+	/** Virtual destructor.*/
+	virtual ~IAutomationControllerManager() { }
 };
diff --git a/Engine/Source/Developer/AutomationController/Public/Interfaces/IAutomationControllerModule.h b/Engine/Source/Developer/AutomationController/Public/Interfaces/IAutomationControllerModule.h
--- a/Engine/Source/Developer/AutomationController/Public/Interfaces/IAutomationControllerModule.h
+++ b/Engine/Source/Developer/AutomationController/Public/Interfaces/IAutomationControllerModule.h
@@ -1,11 +1,8 @@
 // Copyright 1998-2014 Epic Games, Inc. All Rights Reserved.
 
-/*=============================================================================
-	IAutomationControllerModule.h: Declares the IAutomationControllerModule interface.
-=============================================================================*/
-
 #pragma once
 
+
 /**
  * Interface for AutomationController modules.
  */
@@ -21,10 +18,9 @@ public:
 	 */
 	virtual IAutomationControllerManagerRef GetAutomationController( ) = 0;
 
-	/** Init message bus usage */
+	/** Init message bus usage. */
 	virtual void Init() = 0;
 
-	/** Tick function that will execute enabled tests for different device clusters */
+	/** Tick function that will execute enabled tests for different device clusters. */
 	virtual void Tick() = 0;
-
 };
diff --git a/Engine/Source/Developer/AutomationController/Public/Interfaces/IAutomationReport.h b/Engine/Source/Developer/AutomationController/Public/Interfaces/IAutomationReport.h
--- a/Engine/Source/Developer/AutomationController/Public/Interfaces/IAutomationReport.h
+++ b/Engine/Source/Developer/AutomationController/Public/Interfaces/IAutomationReport.h
@@ -1,9 +1,5 @@
 // Copyright 1998-2014 Epic Games, Inc. All Rights Reserved.
 
-/*=============================================================================
-	IAutomationTestStatus.h: Declares the IAutomationTestStatus interface.
-=============================================================================*/
-
 #pragma once
 
 
@@ -150,15 +146,17 @@ public:
 /**
 	 * Returns the name of this level in the test hierarchy for the purposes of grouping.
 	 * (Editor.Maps.LoadAll.parameters) would have 3 internal tree nodes and a variable number of leaf nodes depending on how many maps existed
-	 * @param bAddDecoration - True is the name should have the number of child tests appended
-	 * @return the name of this level in the hierarchy for use in UI
+	 *
+	 * @param bAddDecoration true is the name should have the number of child tests appended.
+	 * @return the name of this level in the hierarchy for use in UI.
 	 */
 	virtual const FString& GetDisplayName() const = 0;
 
 	/**
 	 * Returns the name of this level in the test hierarchy for the purposes of UI.
 	 * (Editor.Maps.LoadAll.parameters) would have 3 internal tree nodes and a variable number of leaf nodes depending on how many maps existed
-	 * @return the name of this level in the hierarchy for use in UI
+	 *
+	 * @return the name of this level in the hierarchy for use in UI.
 	 */
 	virtual FString GetDisplayNameWithDecoration() const = 0;
 
@@ -205,16 +203,16 @@ public:
 	 */
 	virtual const bool IsSmokeTest() = 0;
 
-	/** Sets whether this test should be executed or not */
+	/** Sets whether this test should be executed or not. */
 	virtual void SetEnabled(bool bShouldBeEnabled) = 0;
 
-	/** Sets whether this test is supported on a particular platform */
+	/** Sets whether this test is supported on a particular platform. */
 	virtual void SetSupport(const int32 ClusterIndex) = 0;
 
 	/**
 	 * Set the test type.
 	 * 
-	 * @param InTestType - the type of test.
+	 * @param InTestType The type of test.
 	 */
 	virtual void SetTestType(const uint8 InTestType) = 0;
 
@@ -222,65 +220,70 @@ public:
 	virtual bool IsSupported(const int32 ClusterIndex) const = 0;
 
 	/**
-	 * Filters the visible tests based on search text, execution status, regression test, etc
-	 * @param InFilterDesc - The description of how to decide if this report should be visible
-	 * @param ParentPassedFilter - If the parent passed the filter
-	 * @return whether this report or any of its children passed the filter
+	 * Filters the visible tests based on search text, execution status, regression test, etc.
+	 *
+	 * @param InFilterDesc The description of how to decide if this report should be visible.
+	 * @param ParentPassedFilter If the parent passed the filter.
+	 * @return whether this report or any of its children passed the filter.
 	 */
 	virtual bool SetFilter( TSharedPtr< AutomationFilterCollection > InFilter, const bool ParentPassedFilter = false ) = 0;
 
-	/** Returns the array of child reports that should be visible to the UI based on filtering */
+	/** Returns the array of child reports that should be visible to the UI based on filtering. */
 	virtual TArray<TSharedPtr<IAutomationReport> >& GetFilteredChildren() = 0;
 
-	/** Returns the array of child reports */
+	/** Returns the array of child reports. */
 	virtual TArray<TSharedPtr<IAutomationReport> >& GetChildReports() = 0;
 
-	/** Updates the report when the number of clusters changes */
+	/** Updates the report when the number of clusters changes. */
 	virtual void ClustersUpdated(const int32 NumClusters) = 0;
 
 	/** 
-	 * Recursively resets the report to "needs to be run", clears cached warnings and errors
-	 * @param NumTestPasses - The number of test passes so we know how many results to create
+	 * Recursively resets the report to "needs to be run", clears cached warnings and errors.
+	 *
+	 * @param NumTestPasses The number of test passes so we know how many results to create.
 	 */
 	virtual void ResetForExecution(const int32 NumTestPasses) = 0;
 
 	/**
-	 * Sets the results of the test for use by the UI
-	 * @param ClusterIndex - Index of the platform reporting the results of this test.  See AutomationDeviceClusterManager
-	 * @param PassIndex - Which test pass these results are for
-	 * @param InResults - The new set of results
+	 * Sets the results of the test for use by the UI.
+	 *
+	 * @param ClusterIndex Index of the platform reporting the results of this test.  See AutomationDeviceClusterManager.
+	 * @param PassIndex Which test pass these results are for.
+	 * @param InResults The new set of results.
 	 */
 	virtual void SetResults( const int32 ClusterIndex, const int32 PassIndex, const FAutomationTestResults& InResults ) = 0;
 
 	/**
-	 * Returns completion statistics for this branch of the testing hierarchy
-	 * @param ClusterIndex - Index of the platform reporting the results of this test.  See AutomationDeviceClusterManager
-	 * @param PassIndex - Which test pass to get the status of
-	 * @param OutCompletionState - Collection structure for execution statistics
+	 * Returns completion statistics for this branch of the testing hierarchy.
+	 *
+	 * @param ClusterIndex Index of the platform reporting the results of this test.  See AutomationDeviceClusterManager.
+	 * @param PassIndex Which test pass to get the status of.
+	 * @param OutCompletionState Collection structure for execution statistics.
 	 */
 	virtual void GetCompletionStatus(const int32 ClusterIndex, const int32 PassIndex, FAutomationCompleteState& OutCompletionState) = 0;
 
 	/**
-	 * Returns the state of the test (not run, in process, success, failure)
-	 * @param ClusterIndex - Index of the platform reporting the results of this test.  See AutomationDeviceClusterManager
-	 * @param PassIndex - Which test pass to get the state of
-	 * @return the current state of the test
+	 * Returns the state of the test (not run, in process, success, failure).
+	 *
+	 * @param ClusterIndex Index of the platform reporting the results of this test.  See AutomationDeviceClusterManager.
+	 * @param PassIndex Which test pass to get the state of.
+	 * @return the current state of the test.
 	 */
 	virtual EAutomationState::Type GetState(const int32 ClusterIndex, const int32 PassIndex) const = 0;
 
 	/**
 	 * Gets a copy of errors and warnings that were found
 	 *
-	 * @param ClusterIndex - Index of the platform we are requesting test results for.
-	 * @param PassIndex - Index of the test pass to get the results for.
-	 *
+	 * @param ClusterIndex Index of the platform we are requesting test results for.
+	 * @param PassIndex Index of the test pass to get the results for.
 	 * @return The collection of results for the given cluster index
 	 */
 	virtual const FAutomationTestResults& GetResults( const int32 ClusterIndex, const int32 PassIndex ) = 0;
 
 	/**
-	 * Gets the number of available test results for a given cluster
-	 * @param ClusterIndex - Index of the platform .
+	 * Gets the number of available test results for a given cluster.
+	 *
+	 * @param ClusterIndex Index of the platform.
 	 * @return The number of results available for this cluster
 	 */
 	virtual const int32 GetNumResults( const int32 ClusterIndex ) = 0;
@@ -293,101 +296,106 @@ public:
 	virtual const int32 GetCurrentPassIndex( const int32 ClusterIndex ) = 0;
 
 	/**
-	 * Gets the name of the instance that ran the test
-	 * @param ClusterIndex - Index of the platform reporting the results of this test.  See AutomationDeviceClusterManager
-	 * @return the name of the device
+	 * Gets the name of the instance that ran the test.
+	 *
+	 * @param ClusterIndex - Index of the platform reporting the results of this test.  See AutomationDeviceClusterManager.
+	 * @return the name of the device.
 	 */
 	virtual FString GetGameInstanceName( const int32 ClusterIndex ) = 0;
 
 	/**
-	 * Add a child test to the hierarchy, creating internal tree nodes as needed
-	 * If NewTestName is Editor.Maps.Loadall.TestName, this will create nodes for Editor, Maps, Loadall, and then a leaf node for the testname with the associated commandline
-	 * @param TestInfo - Structure containing all the test info
-	 * @param ClusterIndex - Index of the platform reporting the results of this test.  See AutomationDeviceClusterManager
-	 * @param NumPasses - The number of passes we are going to perform.  Used to make sure we have enough results.
-	 * @return - the automation report
+	 * Add a child test to the hierarchy, creating internal tree nodes as needed.
+	 * If NewTestName is Editor.Maps.Loadall.TestName, this will create nodes for Editor, Maps, Loadall, and then a leaf node for the test name with the associated command line
+	 *
+	 * @param TestInfo Structure containing all the test info.
+	 * @param ClusterIndex Index of the platform reporting the results of this test.  See AutomationDeviceClusterManager.
+	 * @param NumPasses The number of passes we are going to perform.  Used to make sure we have enough results..
+	 * @return The automation report.
 	 */
 	virtual TSharedPtr<IAutomationReport> EnsureReportExists(FAutomationTestInfo& TestInfo, const int32 ClusterIndex, const int32 NumPasses) = 0;
 
 	/**
-	 * Returns the next test in the hierarchy to run
-	 * @param bOutAllTestsComplete - Whether or not all enabled tests have completed execution for this platform (cluster index)
-	 * @param ClusterIndex - Index of the platform reporting the results of this test.  See AutomationDeviceClusterManager
-	 * @param PassIndex - Which test pass we are currently on
-	 * @param NumDevicesInCluster - The number of devices which are in this cluster
+	 * Returns the next test in the hierarchy to run.
 	 *
+	 * @param bOutAllTestsComplete Whether or not all enabled tests have completed execution for this platform (cluster index).
+	 * @param ClusterIndex Index of the platform reporting the results of this test.  See AutomationDeviceClusterManager.
+	 * @param PassIndex Which test pass we are currently on.
+	 * @param NumDevicesInCluster The number of devices which are in this cluster.
 	 * @return The next report
 	 */
 	virtual TSharedPtr<IAutomationReport> GetNextReportToExecute(bool& bOutAllTestsComplete, const int32 ClusterIndex, const int32 PassIndex, const int32 NumDevicesInCluster) = 0;
 
 	/**
-	 * Returns if there have been any errors in the test
-	 * @return True If there were errors
+	 * Returns if there have been any errors in the test.
+	 *
+	 * @return True If there were errors.
 	 */
 	virtual const bool HasErrors() = 0;
 
 	/**
-	 * Returns if there have been any warnings in the test
-	 * @return True If there were errors
+	 * Returns if there have been any warnings in the test.
+	 *
+	 * @return True If there were errors.
 	 */
 	virtual const bool HasWarnings() = 0;
 
 	/**
-	 * Gets the min and max time this test took to execute
-	 * @param OutMinTime - Minimum execution time for all device types
-	 * @param OutMaxTime - Maximum execution time for all device types
-	 * @return whether any test has completed successfully
+	 * Gets the min and max time this test took to execute.
+	 *
+	 * @param OutMinTime Minimum execution time for all device types.
+	 * @param OutMaxTime Maximum execution time for all device types.
+	 * @return whether any test has completed successfully.
 	 */
 	virtual const bool GetDurationRange(float& OutMinTime, float& OutMaxTime) = 0;
 
 	/**
-	 * Get the number of devices which have been given this test to run
+	 * Get the number of devices which have been given this test to run.
 	 *
-	 * @return - The number of devices who have been given this test to run
+	 * @return The number of devices who have been given this test to run.
 	 */
 	virtual const int32 GetNumDevicesRunningTest() const = 0;
 
 	/**
 	 * Get the number of participants this test requires.
 	 *
-	 * @Return - The number of devices needed for this test to execute
+	 * @Return The number of devices needed for this test to execute.
 	 */
 	virtual const int32 GetNumParticipantsRequired() const = 0;
 	
 	/**
 	 * Set the number of participants this test requires if less than what is already set.
 	 *
-	 * @param NewCount - The number of further devices 
+	 * @param NewCount The number of further devices .
 	 */
 	virtual void SetNumParticipantsRequired( const int32 NewCount ) = 0;
 
 	/**
-	 * Increment the number of network responses
-	 * @return - whether we have received responses from all participants
+	 * Increment the number of network responses.
+	 *
+	 * @return true if responses from all participants were received, false otherwise.
 	 */
 	virtual bool IncrementNetworkCommandResponses() = 0;
 
-	/**
-	 * Resets the number of network responses back to zero
-	 */
+	/*** Resets the number of network responses back to zero. */
 	virtual void ResetNetworkCommandResponses() = 0;
 
 	/**
-	 * Should we expand this node in the UI - A child has passed the filter
-	 * @return True If we should expand the node
+	 * Should we expand this node in the UI - A child has passed the filter.
+	 *
+	 * @return true If we should expand the node
 	 */
 	virtual const bool ExpandInUI() const = 0;
 
-	/** Stop the test which is creating this report */
+	/** Stop the test which is creating this report. */
 	virtual void StopRunningTest() = 0;
 
 	/**
-	 * Notification on whether we should, or should not, track this reports history
+	 * Notification on whether we should, or should not, track this reports history.
 	 */
 	virtual void TrackHistory( const bool bShouldTrack, const int32 NumReportsToTrack ) = 0;
 
 	/**
-	 * Get the history items of this particular test
+	 * Get the history items of this particular test.
 	 *
 	 * @return A reference to the items of this tests previous runs.
 	 */
diff --git a/Engine/Source/Developer/AutomationWindow/Private/AutomationFilter.h b/Engine/Source/Developer/AutomationWindow/Private/AutomationFilter.h
--- a/Engine/Source/Developer/AutomationWindow/Private/AutomationFilter.h
+++ b/Engine/Source/Developer/AutomationWindow/Private/AutomationFilter.h
@@ -1,18 +1,14 @@
 // Copyright 1998-2014 Epic Games, Inc. All Rights Reserved.
-#pragma once
 
+#pragma once
 
-/*=============================================================================
-	AutomationFilter.h: Declares the AutomationFilter class.
-=============================================================================*/
 
-class FAutomationFilter : public IFilter< const TSharedPtr< class IAutomationReport >&  >
+class FAutomationFilter
+	: public IFilter<const TSharedPtr<class IAutomationReport>&>
 {
 public:
 
-	/**
-	 * Default constructor.
-	 */
+	/** Default constructor. */
 	FAutomationFilter()
 		: OnlySmokeTests( false )
 		, ShowErrors( false )
@@ -22,8 +18,10 @@ public:
 public:
 
 	/**
-	 * Set if we should only show warnings
-	 * @param InShowWarnings If we should show warnings or not
+	 * Set if we should only show warnings.
+	 *
+	 * @param InShowWarnings If we should show warnings or not.
+	 * @see SetOnlyShowSmokeTests, SetShowErrors, ShouldShowWarnings
 	 */
 	void SetShowWarnings( const bool InShowWarnings )
 	{
@@ -31,8 +29,10 @@ public:
 	}
 
 	/**
-	 * Should we show warnings
-	 * @return True if we should be showing warnings
+	 * Should we show warnings.
+	 *
+	 * @return True if we should be showing warnings.
+	 * @see SetShowErrors, SetShowWarnings
 	 */
 	const bool ShouldShowWarnings( ) const
 	{
@@ -40,8 +40,10 @@ public:
 	}
 
 	/**
-	 * Set if we should only show errors
-	 * @param InShowErrors If we should show errors
+	 * Set if we should only show errors.
+	 *
+	 * @param InShowErrors If we should show errors.
+	 * @see SetOnlyShowSmokeTests, SetShowWarnings, ShouldShowErrors
 	 */
 	void SetShowErrors( const bool InShowErrors )
 	{
@@ -49,8 +51,10 @@ public:
 	}
 
 	/**
-	 * Should we show errors
-	 * @return True if we should be showing errors
+	 * Should we show errors.
+	 *
+	 * @return True if we should be showing errors.
+	 * @see SetShowErrors, ShouldShowWarnings
 	 */
 	const bool ShouldShowErrors( ) const
 	{
@@ -58,8 +62,9 @@ public:
 	}
 
 	/**
-	 * Set if we should only show smoke tests
-	 * @param InOnlySmokeTests If we should show smoke tests
+	 * Set if we should only show smoke tests.
+	 *
+	 * @param InOnlySmokeTests If we should show smoke tests.
 	 */
 	void SetOnlyShowSmokeTests( const bool InOnlySmokeTests )
 	{
@@ -67,8 +72,10 @@ public:
 	}
 
 	/**
-	 * Should we only show errors
-	 * @return True if we should be showing errors
+	 * Should we only show errors.
+	 *
+	 * @return True if we should be showing errors.
+	 * @see SetOnlyShowSmokeTests, ShouldShowErrors, ShouldShowWarnings
 	 */
 	const bool OnlyShowSmokeTests( ) const
 	{
@@ -109,15 +116,15 @@ public:
 
 private:
 
-	/**	The event that broadcasts whenever a change occurs to the filter */
+	/**	The event that broadcasts whenever a change occurs to the filter. */
 	FChangedEvent ChangedEvent;
 	
-	/** Only smoke tests will pass the test **/
+	/** Only smoke tests will pass the test. */
 	bool OnlySmokeTests;
 
-	/** Only errors will pass the test **/
+	/** Only errors will pass the test. */
 	bool ShowErrors;
 
-	/** Only warnings will pass the test **/
+	/** Only warnings will pass the test. */
 	bool ShowWarnings;
 };
\ No newline at end of file
diff --git a/Engine/Source/Developer/AutomationWindow/Private/AutomationPresetManager.cpp b/Engine/Source/Developer/AutomationWindow/Private/AutomationPresetManager.cpp
--- a/Engine/Source/Developer/AutomationWindow/Private/AutomationPresetManager.cpp
+++ b/Engine/Source/Developer/AutomationWindow/Private/AutomationPresetManager.cpp
@@ -2,10 +2,12 @@
 
 #include "AutomationWindowPrivatePCH.h"
 
+
 FAutomationTestPresetManager::FAutomationTestPresetManager()
 {
 }
 
+
 AutomationPresetRef FAutomationTestPresetManager::AddNewPreset( )
 {
 	AutomationPresetPtr NewPreset = MakeShareable(new FAutomationTestPreset());
@@ -13,6 +15,7 @@ AutomationPresetRef FAutomationTestPresetManager::AddNewPreset( )
 	return NewPreset.ToSharedRef();
 }
 
+
 AutomationPresetRef FAutomationTestPresetManager::AddNewPreset( const FString& PresetName, const TArray<FString>& SelectedTests )
 {
 	AutomationPresetRef NewPreset = AddNewPreset();
@@ -22,11 +25,13 @@ AutomationPresetRef FAutomationTestPresetManager::AddNewPreset( const FString& P
 	return NewPreset;
 }
 
+
 TArray<AutomationPresetPtr>& FAutomationTestPresetManager::GetAllPresets( )
 {
 	return Presets;
 }
 
+
 AutomationPresetPtr FAutomationTestPresetManager::LoadPreset( FArchive& Archive )
 {
 	FAutomationTestPreset* NewPreset = new FAutomationTestPreset();
@@ -38,9 +43,10 @@ AutomationPresetPtr FAutomationTestPresetManager::LoadPreset( FArchive& Archive
 
 	delete NewPreset;
 
-	return NULL;
+	return nullptr;
 }
 
+
 void FAutomationTestPresetManager::RemovePreset( const AutomationPresetRef Preset )
 {
 	if (Presets.Remove(Preset) > 0)
@@ -53,12 +59,13 @@ void FAutomationTestPresetManager::RemovePreset( const AutomationPresetRef Prese
 	}
 }
 
+
 void FAutomationTestPresetManager::SavePreset( const AutomationPresetRef Preset )
 {
 	FString PresetFileName = GetPresetFolder() / Preset->GetID().ToString() + TEXT(".uap");
 	FArchive* PresetFileWriter = IFileManager::Get().CreateFileWriter(*PresetFileName);
 
-	if (PresetFileWriter != NULL)
+	if (PresetFileWriter != nullptr)
 	{
 		SavePreset(Preset, *PresetFileWriter);
 
@@ -66,11 +73,13 @@ void FAutomationTestPresetManager::SavePreset( const AutomationPresetRef Preset
 	}
 }
 
+
 void FAutomationTestPresetManager::SavePreset( const AutomationPresetRef Preset, FArchive& Archive )
 {
 	Preset->Serialize(Archive);
 }
 
+
 void FAutomationTestPresetManager::LoadPresets( )
 {
 	TArray<FString> PresetFileNames;
@@ -82,7 +91,7 @@ void FAutomationTestPresetManager::LoadPresets( )
 		FString PresetFilePath = GetPresetFolder() / *It;
 		FArchive* PresetFileReader = IFileManager::Get().CreateFileReader(*PresetFilePath);
 
-		if (PresetFileReader != NULL)
+		if (PresetFileReader != nullptr)
 		{
 			AutomationPresetPtr LoadedPreset = LoadPreset(*PresetFileReader);
 
diff --git a/Engine/Source/Developer/AutomationWindow/Private/AutomationPresetManager.h b/Engine/Source/Developer/AutomationWindow/Private/AutomationPresetManager.h
--- a/Engine/Source/Developer/AutomationWindow/Private/AutomationPresetManager.h
+++ b/Engine/Source/Developer/AutomationWindow/Private/AutomationPresetManager.h
@@ -2,19 +2,17 @@
 
 #pragma once
 
-/**
- * Type definition for shared pointers to instances of FAutomationTestPreset.
- */
+
+/** Type definition for shared pointers to instances of FAutomationTestPreset. */
 typedef TSharedPtr<class FAutomationTestPreset> AutomationPresetPtr;
 
-/**
- * Type definition for shared references to instances of FAutomationTestPreset.
- */
+/** Type definition for shared references to instances of FAutomationTestPreset. */
 typedef TSharedRef<class FAutomationTestPreset> AutomationPresetRef;
 
 //Increment this number if the Preset serialization changes.
 #define AUTOMATION_PRESETVERSION 1
 
+
 /**
  * Class that holds preset data for the automation window
  */
@@ -22,21 +20,19 @@ class FAutomationTestPreset
 {
 public:
 
-	/**
-	 * Constructor.
-	 */
+	/** Default constructor. */
 	FAutomationTestPreset() :
 		Id(FGuid::NewGuid())
-	{
-	}
+	{ }
 
-	/**
-	 * Destructor.
-	 */
+	/** Destructor. */
 	virtual ~FAutomationTestPreset( ) { }
 
 	/**
-	 * Gets the GUid for this preset
+	 * Gets the GUid for this preset.
+	 *
+	 * @return The preset's unique identifier.
+	 * @see GetPresetName
 	 */
 	const FGuid GetID()
 	{
@@ -44,7 +40,10 @@ public:
 	}
 
 	/**
-	 * Gets the name for this preset
+	 * Gets the name for this preset.
+	 *
+	 * @return The preset's name.
+	 * @see GetID, SetPresetName
 	 */
 	const FString& GetPresetName( ) const
 	{
@@ -52,7 +51,10 @@ public:
 	}
 
 	/**
-	 * Sets the name for this preset
+	 * Sets the name for this preset.
+	 *
+	 * @param InPresetName The name to set.
+	 * @see GetPresetName
 	 */
 	void SetPresetName(const FString& InPresetName)
 	{
@@ -60,7 +62,10 @@ public:
 	}
 
 	/**
-	 * Returns the list of enabled tests
+	 * Returns the list of enabled tests.
+	 *
+	 * @return Collection of tests.
+	 * @see SetEnabledTests
 	 */
 	const TArray<FString>& GetEnabledTests() const
 	{
@@ -68,7 +73,10 @@ public:
 	}
 
 	/**
-	 * Sets the list of enabled tests
+	 * Sets the list of enabled tests.
+	 *
+	 * @param NewEnableDtests The collection of tests to enable.
+	 * @see GetEnabledTests
 	 */
 	void SetEnabledTests(TArray<FString> NewEnabledTests)
 	{
@@ -76,7 +84,9 @@ public:
 	}
 
 	/**
-	 * Handles saving / loading the preset data to an archive (file)
+	 * Handles saving / loading the preset data to an archive (file).
+	 *
+	 * @param Archive The archive to serialize from/to.
 	 */
 	virtual bool Serialize( FArchive& Archive )
 	{
@@ -97,72 +107,82 @@ public:
 	}
 
 private:
-	// A unique ID for this preset. (Used for the filename because the PresetName may have invalid characters)
+
+	/** A unique ID for this preset. (Used for the filename because the PresetName may have invalid characters). */
 	FGuid Id;
 
-	// The name of this preset
+	/** The name of this preset. */
 	FString PresetName;
 
-	// The list of enabled test names
+	/** The list of enabled test names. */
 	TArray<FString> EnabledTests;
 };
 
+
 class FAutomationTestPresetManager
 {
 public:
-	/**
-	 * Constructor.
-	 */
+
+	/** Default constructor. */
 	FAutomationTestPresetManager();
 
-	virtual ~FAutomationTestPresetManager() {}
+	/** Virtual destructor. */
+	virtual ~FAutomationTestPresetManager() { }
 
 	/**
-	 * Creates a new empty preset
+	 * Creates a new empty preset.
 	 */
 	virtual AutomationPresetRef AddNewPreset( );
 
 	/**
-	 * Creates a new preset with the given name and enabled tests
-	 * @param PresetName - The name of the new preset
-	 * @param SelectedTests - The list of enabled tests
+	 * Creates a new preset with the given name and enabled tests.
+	 * @param PresetName The name of the new preset.
+	 * @param SelectedTests The list of enabled tests.
 	 */
 	virtual AutomationPresetRef AddNewPreset( const FString& PresetName, const TArray<FString>& SelectedTests );
 
 	/**
-	 * Returns a reference to the list that holds the presets
+	 * Returns a reference to the list that holds the presets.
 	 */
 	virtual TArray<AutomationPresetPtr>& GetAllPresets( );
 
 	/**
-	 * Removes the selected preset from the preset list and deletes the file on disk
-	 * @param Preset - The preset to remove
+	 * Removes the selected preset from the preset list and deletes the file on disk.
+	 *
+	 * @param Preset The preset to remove.
+	 * @see LoadPresets, SavePreset
 	 */
 	void RemovePreset( const AutomationPresetRef Preset );
 
 	/**
-	 * Saves the passed in preset to disk
-	 * @param Preset - The preset to save
+	 * Saves the passed in preset to disk.
+	 *
+	 * @param Preset The preset to save.
+	 * @see RemovePreset, LoadPresets
 	 */
 	void SavePreset( const AutomationPresetRef Preset);
 
 	/**
-	* Load all Presets from disk.
-	*/
+	 * Load all Presets from disk.
+	 *
+	 * @see RemovePreset, SavePreset
+	 */
 	void LoadPresets( );
 
 protected:
 
 	/**
-	 * Creates a new preset and loads it with data from the archive
-	 * @param Archive - The stream to read the preset data from
+	 * Creates a new preset and loads it with data from the archive.
+	 *
+	 * @param Archive The stream to read the preset data from.
 	 */
 	AutomationPresetPtr LoadPreset( FArchive& Archive );
 
 	/**
-	 * Writes the Preset data to the archive
-	 * @param Preset - The preset to serialize
-	 * @param Archive - The stream to write the preset data to
+	 * Writes the Preset data to the archive.
+	 *
+	 * @param Preset The preset to serialize.
+	 * @param Archive The stream to write the preset data to.
 	 */
 	void SavePreset( const AutomationPresetRef Preset, FArchive& Archive );
 
@@ -178,6 +198,6 @@ protected:
 
 private:
 
-	// Holds the collection of automation presets.
+	/** Holds the collection of automation presets. */
 	TArray<AutomationPresetPtr> Presets;
 };
diff --git a/Engine/Source/Developer/AutomationWindow/Private/AutomationWindowModule.cpp b/Engine/Source/Developer/AutomationWindow/Private/AutomationWindowModule.cpp
--- a/Engine/Source/Developer/AutomationWindow/Private/AutomationWindowModule.cpp
+++ b/Engine/Source/Developer/AutomationWindow/Private/AutomationWindowModule.cpp
@@ -1,9 +1,5 @@
 // Copyright 1998-2014 Epic Games, Inc. All Rights Reserved.
 
-/*=============================================================================
-	AutomationWindowModule.cpp: Implements the FAutomationWindowModule class.
-=============================================================================*/
-
 #include "AutomationWindowPrivatePCH.h"
 
 
@@ -57,4 +53,4 @@ private:
 };
 
 
-IMPLEMENT_MODULE(FAutomationWindowModule, AutomationWindow);
\ No newline at end of file
+IMPLEMENT_MODULE(FAutomationWindowModule, AutomationWindow);
diff --git a/Engine/Source/Developer/AutomationWindow/Private/SAutomationExportMenu.h b/Engine/Source/Developer/AutomationWindow/Private/SAutomationExportMenu.h
--- a/Engine/Source/Developer/AutomationWindow/Private/SAutomationExportMenu.h
+++ b/Engine/Source/Developer/AutomationWindow/Private/SAutomationExportMenu.h
@@ -1,15 +1,13 @@
 // Copyright 1998-2014 Epic Games, Inc. All Rights Reserved.
+
 #pragma once
-/*=============================================================================
-	SAutomationExportMenu.h: Declares the AutomationExortMenu class.
-=============================================================================*/
+
 
 class SAutomationExportMenu : public SCompoundWidget
 {
-
 public:
 
-	// Constructors
+	/** Default constructor. */
 	SAutomationExportMenu();
 
 	SLATE_BEGIN_ARGS( SAutomationExportMenu ){}
@@ -25,53 +23,49 @@ public:
 private:
 
 	/**
-	* Have the reports been generated
-	*
-	* @return true if the reports have been generated
-	*/
+	 * Have the reports been generated
+	 *
+	 * @return true if the reports have been generated
+	 */
 	bool AreReportsGenerated() const;
 
 	/** 
-	* Build the menu items.
-	*
-	* @param In Name - menu item name
-	* @param InType - menu item type
-	*/
+	 * Build the menu items.
+	 *
+	 * @param InName The menu item name.
+	 * @param InType The menu item type.
+	 */
 	void BuildMenuItems( const FText& InName, EFileExportType::Type InType );
 
-	/**
-	* Create the menu widgets.
-	*/
+	/** Create the menu widgets. */
 	void CreateMenu( );
 
 	/**
-	* Set all the available export options to checked.
-	*/
+	 * Set all the available export options to checked.
+	 */
 	void EnableAvailableReports();
 
 	/**
-	* Get the export button tooltip
-	*
-	* @return true export report is ready
-	*/
+	 * Get the export button tooltip.
+	 *
+	 * @return true export report is ready.
+	 */
 	FText GetExportButtonTooltip() const;
 
 	/**
-	* Get the export combo button tooltop
-	*
-	* @return true export report is ready
-	*/
+	 * Get the export combo button tooltip.
+	 *
+	 * @return true export report is ready.
+	 */
 	FText GetExportComboButtonTooltip() const;
 
-	/**
-	* Get the test result summary from the manager.
-	*/
+	/** Get the test result summary from the manager. */
 	void GetResults( );
 
 	/** 
 	* Export the report data when clicked.
 	*
-	* @return handled if the report has been handled
+	* @return handled if the report has been handled.
 	*/
 	FReply HandleExportDataClicked();
 
@@ -81,41 +75,41 @@ private:
 	void HandleMenuOpen();
 
 	/** 
-	* Should the export checkboxes be enabled.
-	*
-	* @param CheckType - the menu item type
-	* @return - True if the checkbox is enabled 
-	*/
+	 * Should the export checkboxes be enabled.
+	 *
+	 * @param CheckType The menu item type.
+	 * @return true if the checkbox is enabled .
+	 */
 	bool IsCheckBoxEnabled( EFileExportType::Type CheckType ) const;
 
 	/**
-	* Is an export report ready
+	* Is an export report ready.
 	*
-	* @return true export report is ready
+	* @return true export report is ready.
 	*/
 	bool IsExportReady() const;
 
 	/** 
-	* Check box has changed, update the export data mask 
-	*
-	* @param The new state, checked or unchecked
-	* @param The type of checkbox set
-	*/
+	 * Check box has changed, update the export data mask .
+	 *
+	 * @param The new state, checked or unchecked.
+	 * @param The type of checkbox set.
+	 */
 	void OnDisplayCheckStateChanged( ESlateCheckBoxState::Type InNewState, EFileExportType::Type CheckType );
 
 	/** 
-	* Gets the display state to send to a display filter check box
-	*
-	* @param - the type of checkbox
-	* @return - The desired checkbox state
-	*/
+	 * Gets the display state to send to a display filter check box
+	 *
+	 * @param The type of checkbox.
+	 * @return The desired checkbox state.
+	 */
 	ESlateCheckBoxState::Type OnGetDisplayCheckState( EFileExportType::Type CheckType ) const;
 
 	/**
-	* Add a notification when the file is exported
-	*
-	* @return true export report is ready
-	*/
+	 * Add a notification when the file is exported.
+	 *
+	 * @return true export report is ready.
+	 */
 	FReply SpawnNotification();
 
 private:
diff --git a/Engine/Source/Developer/AutomationWindow/Private/SAutomationGraphicalResultBox.cpp b/Engine/Source/Developer/AutomationWindow/Private/SAutomationGraphicalResultBox.cpp
--- a/Engine/Source/Developer/AutomationWindow/Private/SAutomationGraphicalResultBox.cpp
+++ b/Engine/Source/Developer/AutomationWindow/Private/SAutomationGraphicalResultBox.cpp
@@ -1,13 +1,11 @@
 // Copyright 1998-2014 Epic Games, Inc. All Rights Reserved.
 
-/*=============================================================================
-	SAutomationGraphicalResultBox.cpp: Implements the SAutomationGraphicalResultBox class.
-=============================================================================*/
-
 #include "AutomationWindowPrivatePCH.h"
 
+
 #define LOCTEXT_NAMESPACE "SAutomationGraphicalResultBox"
 
+
 /* SAutomationGraphicalResultBox interface
  *****************************************************************************/
 
@@ -33,6 +31,7 @@ void SAutomationGraphicalResultBox::Construct( const FArguments& InArgs, const I
 	AutomationController->SetTestsCompleteCallback(FOnAutomationControllerTestsComplete::CreateRaw(this, &SAutomationGraphicalResultBox::OnFillResults));
 }
 
+
 void SAutomationGraphicalResultBox::ClearResults()
 {
 	TotalTestDuration = 0.f;
@@ -40,6 +39,7 @@ void SAutomationGraphicalResultBox::ClearResults()
 	RootBox->ClearChildren();
 }
 
+
 void SAutomationGraphicalResultBox::GetEnabledReports(TSharedPtr<IAutomationReport> InReport, TArray< TSharedPtr< IAutomationReport > >& outReports)
 {
 	TArray<TSharedPtr<IAutomationReport> >& ChildReports = InReport->GetChildReports();
@@ -57,6 +57,7 @@ void SAutomationGraphicalResultBox::GetEnabledReports(TSharedPtr<IAutomationRepo
 	}
 }
 
+
 void SAutomationGraphicalResultBox::PopulateData()
 {
 	//Get the report list
@@ -119,18 +120,15 @@ void SAutomationGraphicalResultBox::PopulateData()
 	}
 }
 
+
 void SAutomationGraphicalResultBox::OnFillResults()
 {
-	//Clear the old results
 	ClearResults();
-
-	//Get the data out of the automation reports
 	PopulateData();
-
-	//Create widgets for the results
 	CreateWidgets();
 }
 
+
 void SAutomationGraphicalResultBox::CreateWidgets()
 {
 	//Note: I am using 10 columns here to hack around a sizing bug.  The Grid calculates the width for items that span multiple 
@@ -271,6 +269,7 @@ void SAutomationGraphicalResultBox::CreateWidgets()
 		];
 }
 
+
 FSlateColor SAutomationGraphicalResultBox::GetColorForTestState(const EAutomationState::Type TestState, const bool bHasWarnings) const
 {
 	switch(TestState)
@@ -293,6 +292,7 @@ FSlateColor SAutomationGraphicalResultBox::GetColorForTestState(const EAutomatio
 	return FSlateColor( FLinearColor( 0.5f, 0.0f, 0.0f ) );
 }
 
+
 FText SAutomationGraphicalResultBox::GetTestDisplayText(const FString TestName, const float TestTime) const
 {
 	if(DisplayType == EAutomationGrapicalDisplayType::DisplayName)
@@ -307,19 +307,23 @@ FText SAutomationGraphicalResultBox::GetTestDisplayText(const FString TestName,
 	}
 }
 
+
 EAutomationGrapicalDisplayType::Type SAutomationGraphicalResultBox::GetDisplayType() const
 {
 	return DisplayType;
 }
 
+
 void SAutomationGraphicalResultBox::SetDisplayType(EAutomationGrapicalDisplayType::Type NewDisplayType)
 {
 	DisplayType = NewDisplayType;
 }
 
+
 bool SAutomationGraphicalResultBox::HasResults() const
 {
 	return ClusterResults.Num() > 0;
 }
 
+
 #undef LOCTEXT_NAMESPACE
