Fixed logic error, found via static analysis

diff --git a/Engine/Source/Developer/AutomationController/Private/AutomationReport.cpp b/Engine/Source/Developer/AutomationController/Private/AutomationReport.cpp
--- a/Engine/Source/Developer/AutomationController/Private/AutomationReport.cpp
+++ b/Engine/Source/Developer/AutomationController/Private/AutomationReport.cpp
@@ -296,7 +296,7 @@ void FAutomationReport::ClustersUpdated(const int32 NumClusters)
 			Results.Add( AutomationTestResult );
 		}
 	}
-	else if( NumClusters > Results.Num() )
+	else if( NumClusters < Results.Num() )
 	{
 		Results.RemoveAt(NumClusters, Results.Num() - NumClusters);
 	}
