Commandline test names no longer need spaces.

Example:
Both of these will work!
     -ExecCmds="Automation RunTests MeshFactoryValidation"
     -ExecCmds="Automation RunTests Mesh Factory Validation"

diff --git a/Engine/Source/Developer/AutomationController/Private/AutomationCommandline.cpp b/Engine/Source/Developer/AutomationController/Private/AutomationCommandline.cpp
--- a/Engine/Source/Developer/AutomationController/Private/AutomationCommandline.cpp
+++ b/Engine/Source/Developer/AutomationController/Private/AutomationCommandline.cpp
@@ -95,8 +95,9 @@ public:
 
 	void GenerateTestNamesFromCommandLine(const TArray<FString>& AllTestNames, TArray<FString>& OutTestNames)
 	{
+		FString AllTestNamesNoWhiteSpaces;
 		OutTestNames.Empty();
-
+		
 		//Split the test names up
 		TArray<FString> Filters;
 		StringCommand.ParseIntoArray(Filters, TEXT("+"), true);
@@ -104,13 +105,16 @@ public:
 		for (int32 FilterIndex = 0; FilterIndex < Filters.Num(); ++FilterIndex)
 		{
 			Filters[FilterIndex] = Filters[FilterIndex].Trim();
+			Filters[FilterIndex] = Filters[FilterIndex].Replace(TEXT(" "), TEXT(""));
 		}
 
 		for (int32 TestIndex = 0; TestIndex < AllTestNames.Num(); ++TestIndex)
 		{
 			for (int FilterIndex = 0; FilterIndex < Filters.Num(); ++FilterIndex)
 			{
-				if (AllTestNames[TestIndex].Contains(Filters[FilterIndex]))
+				AllTestNamesNoWhiteSpaces = AllTestNames[TestIndex];
+				AllTestNamesNoWhiteSpaces = AllTestNamesNoWhiteSpaces.Replace(TEXT(" "), TEXT(""));
+				if (AllTestNamesNoWhiteSpaces.Contains(Filters[FilterIndex]))
 				{
 					OutTestNames.Add(AllTestNames[TestIndex]);
 					TestCount++;
