Switching the automation commandline exec to use GLog as its OutputDevice to prevent a crash when it is used later (as the command executes over multiple frames).

diff --git a/Engine/Source/Developer/AutomationController/Private/AutomationCommandline.cpp b/Engine/Source/Developer/AutomationController/Private/AutomationCommandline.cpp
--- a/Engine/Source/Developer/AutomationController/Private/AutomationCommandline.cpp
+++ b/Engine/Source/Developer/AutomationController/Private/AutomationCommandline.cpp
@@ -269,7 +269,7 @@ public:
 		bool bHandled = false;
 
 		//save off device to send results to
-		OutputDevice = &Ar;
+		OutputDevice = GLog;
 
 		StringCommand.Empty();
 		//figure out if we are handling this request
