Automation - Adding the basics of automation reporting, it now dumps and html file and several images representing artifacts generated during the when the build ran.

diff --git a/Engine/Source/Developer/AutomationController/Private/AutomationControllerManager.h b/Engine/Source/Developer/AutomationController/Private/AutomationControllerManager.h
--- a/Engine/Source/Developer/AutomationController/Private/AutomationControllerManager.h
+++ b/Engine/Source/Developer/AutomationController/Private/AutomationControllerManager.h
@@ -23,24 +23,31 @@ struct FAutomatedTestResult
 	GENERATED_BODY()
 public:
 
+	TSharedPtr<IAutomationReport> Test;
+
+	UPROPERTY()
+	FString TestDisplayName;
+
 	UPROPERTY()
-	FString TestName;
+	FString FullTestPath;
 
 	UPROPERTY()
-	FString TestResult;
+	EAutomationState State;
 
 	UPROPERTY()
-	TArray<FString> TestInfo;
+	TArray<FString> Logs;
 
 	UPROPERTY()
-	TArray<FString> TestErrors;
+	TArray<FString> Errors;
 
 	UPROPERTY()
-	TArray<FString> TestWarnings;
+	TArray<FString> Warnings;
+
+	TArray<FAutomationArtifact> Artifacts;
 
 	FAutomatedTestResult()
 	{
-		TestResult = TEXT("Not Run");
+		State = EAutomationState::NotRun;
 	}
 };
 
@@ -85,8 +92,9 @@ public:
 class FAutomationControllerManager : public IAutomationControllerManager
 {
 public:
-	// IAutomationController Interface
+	FAutomationControllerManager();
 
+	// IAutomationController Interface
 	virtual void RequestAvailableWorkers( const FGuid& InSessionId ) override;
 	virtual void RequestTests() override;
 	virtual void RunTests( const bool bIsLocalSession) override;
@@ -290,22 +298,13 @@ protected:
 	 */
 	void GenerateHtmlTestPassSummary(FDateTime Timestamp);
 
-	/**
-	* Updates the result value of a finished test.
-	*
-	* @param TestName The test that was just run.
-	* @param TestResult The result that the test returned.
-	*/
-	void UpdateTestResultValue(FString TestName, EAutomationState::Type TestResult);
-
 	/**
 	* Gather all info, warning, and error lines generated over the course of a test.
 	*
 	* @param TestName The test that was just run.
 	* @param TestResult All of the messages of note generated during the test case.
 	*/
-	void CollectTestNotes(FString TestName, const FAutomationWorkerRunTestsReply& Message);
-
+	void CollectTestResults(TSharedPtr<IAutomationReport> Report, const FAutomationTestResults& Results);
 
 	/**
 	 * Checks the child result.
@@ -491,6 +490,9 @@ private:
 	/** Pending image comparisons */
 	TQueue<TSharedPtr<FComparisonEntry>> ComparisonQueue;
 
+	/** The report folder override path that may have been provided over the commandline, -ReportOutputPath="" */
+	FString ReportOutputPathOverride;
+
 private:
 
 	/** Holds a delegate that is invoked when the controller shuts down. */
diff --git a/Engine/Source/Developer/AutomationController/Private/AutomationControllerManger.cpp b/Engine/Source/Developer/AutomationController/Private/AutomationControllerManger.cpp
--- a/Engine/Source/Developer/AutomationController/Private/AutomationControllerManger.cpp
+++ b/Engine/Source/Developer/AutomationController/Private/AutomationControllerManger.cpp
@@ -32,6 +32,11 @@ namespace AutomationControllerConstants
 	const FString HistoryConfigSectionName = TEXT("AutomationController.History");
 }
 
+FAutomationControllerManager::FAutomationControllerManager()
+{
+	FParse::Value(FCommandLine::Get(), TEXT("ReportOutputPath="), ReportOutputPathOverride, false);
+}
+
 void FAutomationControllerManager::RequestAvailableWorkers(const FGuid& SessionId)
 {
 	//invalidate previous tests
@@ -69,10 +74,10 @@ void FAutomationControllerManager::RequestTests()
 
 	ReportManager.Empty();
 
-	for (int32 ClusterIndex = 0; ClusterIndex < DeviceClusterManager.GetNumClusters(); ++ClusterIndex)
+	for ( int32 ClusterIndex = 0; ClusterIndex < DeviceClusterManager.GetNumClusters(); ++ClusterIndex )
 	{
 		int32 DevicesInCluster = DeviceClusterManager.GetNumDevicesInCluster(ClusterIndex);
-		if (DevicesInCluster > 0)
+		if ( DevicesInCluster > 0 )
 		{
 			FMessageAddress MessageAddress = DeviceClusterManager.GetDeviceMessageAddress(ClusterIndex, 0);
 
@@ -84,8 +89,7 @@ void FAutomationControllerManager::RequestTests()
 	}
 }
 
-
-void FAutomationControllerManager::RunTests( const bool bInIsLocalSession )
+void FAutomationControllerManager::RunTests(const bool bInIsLocalSession)
 {
 	ExecutionCount++;
 	CurrentTestPass = 0;
@@ -110,13 +114,13 @@ void FAutomationControllerManager::RunTests( const bool bInIsLocalSession )
 	//reset all tests
 	ReportManager.ResetForExecution(NumTestPasses);
 
-	for (int32 ClusterIndex = 0; ClusterIndex < DeviceClusterManager.GetNumClusters(); ++ClusterIndex)
+	for ( int32 ClusterIndex = 0; ClusterIndex < DeviceClusterManager.GetNumClusters(); ++ClusterIndex )
 	{
 		//enable each device cluster
-		ClusterDistributionMask |= (1<<ClusterIndex);
+		ClusterDistributionMask |= ( 1 << ClusterIndex );
 
 		//for each device in this cluster
-		for (int32 DeviceIndex = 0; DeviceIndex < DeviceClusterManager.GetNumDevicesInCluster(ClusterIndex); ++DeviceIndex)
+		for ( int32 DeviceIndex = 0; DeviceIndex < DeviceClusterManager.GetNumDevicesInCluster(ClusterIndex); ++DeviceIndex )
 		{
 			//mark the device as idle
 			DeviceClusterManager.SetTest(ClusterIndex, DeviceIndex, NULL);
@@ -128,13 +132,12 @@ void FAutomationControllerManager::RunTests( const bool bInIsLocalSession )
 	}
 
 	// Inform the UI we are running tests
-	if (ClusterDistributionMask != 0)
+	if ( ClusterDistributionMask != 0 )
 	{
-		SetControllerStatus( EAutomationControllerModuleState::Running );
+		SetControllerStatus(EAutomationControllerModuleState::Running);
 	}
 }
 
-
 void FAutomationControllerManager::StopTests()
 {
 	bTestResultsAvailable = false;
@@ -143,19 +146,18 @@ void FAutomationControllerManager::StopTests()
 	ReportManager.StopRunningTests();
 
 	// Inform the UI we have stopped running tests
-	if (DeviceClusterManager.HasActiveDevice())
+	if ( DeviceClusterManager.HasActiveDevice() )
 	{
-		SetControllerStatus( EAutomationControllerModuleState::Ready );
+		SetControllerStatus(EAutomationControllerModuleState::Ready);
 	}
 	else
 	{
-		SetControllerStatus( EAutomationControllerModuleState::Disabled );
+		SetControllerStatus(EAutomationControllerModuleState::Disabled);
 	}
 
 	TestRunningArray.Empty();
 }
 
-
 void FAutomationControllerManager::Init()
 {
 	extern void EmptyLinkFunctionForStaticInitializationAutomationExecCmd();
@@ -175,13 +177,11 @@ void FAutomationControllerManager::Init()
 	GConfig->GetInt(*AutomationControllerConstants::HistoryConfigSectionName, TEXT("NumberOfHistoryItemsTracked"), NumberOfHistoryItemsTracked, GEngineIni);
 }
 
-
-void FAutomationControllerManager::RequestLoadAsset( const FString& InAssetName )
+void FAutomationControllerManager::RequestLoadAsset(const FString& InAssetName)
 {
 	MessageEndpoint->Publish(new FAssetEditorRequestOpenAsset(InAssetName), EMessageScope::Process);
 }
 
-
 void FAutomationControllerManager::Tick()
 {
 	ProcessAvailableTasks();
@@ -205,6 +205,26 @@ void FAutomationControllerManager::ProcessComparisonQueue()
 
 			// Issue tests on appropriate platforms
 			MessageEndpoint->Send(new FAutomationWorkerImageComparisonResults(bIsNew, bAreSimilar), Entry->Sender);
+
+			// Record the metadata for the test that needed the screenshots compared.
+			for ( int32 Index = 0; Index < TestRunningArray.Num(); Index++ )
+			{
+				// Find the game session instance info
+				int32 ClusterIndex;
+				int32 DeviceIndex;
+				verify(DeviceClusterManager.FindDevice(Entry->Sender, ClusterIndex, DeviceIndex));
+
+				TSharedPtr<IAutomationReport> Report = DeviceClusterManager.GetTest(ClusterIndex, DeviceIndex);
+				check(Report.IsValid());
+
+				FString ApprovedFolder = ScreenshotManager->GetLocalApprovedFolder();
+				FString UnapprovedFolder = ScreenshotManager->GetLocalUnapprovedFolder();
+				FString ComparisonFolder = ScreenshotManager->GetLocalComparisonFolder();
+
+				Report->AddArtifact(ClusterIndex, CurrentTestPass, FAutomationArtifact(TEXT("Approved"), ApprovedFolder / Result.ApprovedFile));
+				Report->AddArtifact(ClusterIndex, CurrentTestPass, FAutomationArtifact(TEXT("Incoming"), UnapprovedFolder / Result.IncomingFile));
+				Report->AddArtifact(ClusterIndex, CurrentTestPass, FAutomationArtifact(TEXT("Difference"), ComparisonFolder / Result.ComparisonFile));
+			}
 		}
 	}
 }
@@ -212,24 +232,24 @@ void FAutomationControllerManager::ProcessComparisonQueue()
 void FAutomationControllerManager::ProcessAvailableTasks()
 {
 	// Distribute tasks
-	if( ClusterDistributionMask != 0)
+	if ( ClusterDistributionMask != 0 )
 	{
 		// For each device cluster
-		for( int32 ClusterIndex = 0; ClusterIndex < DeviceClusterManager.GetNumClusters(); ++ClusterIndex )
+		for ( int32 ClusterIndex = 0; ClusterIndex < DeviceClusterManager.GetNumClusters(); ++ClusterIndex )
 		{
 			bool bAllTestsComplete = true;
 
 			// If any of the devices were valid
-			if( ( ClusterDistributionMask & ( 1<< ClusterIndex ) ) && DeviceClusterManager.GetNumDevicesInCluster( ClusterIndex ) > 0 )
+			if ( ( ClusterDistributionMask & ( 1 << ClusterIndex ) ) && DeviceClusterManager.GetNumDevicesInCluster(ClusterIndex) > 0 )
 			{
-				ExecuteNextTask( ClusterIndex, bAllTestsComplete );
+				ExecuteNextTask(ClusterIndex, bAllTestsComplete);
 			}
 
 			//if we're all done running our tests
 			if ( bAllTestsComplete )
 			{
 				//we don't need to test this cluster anymore
-				ClusterDistributionMask &= ~( 1<<ClusterIndex );
+				ClusterDistributionMask &= ~( 1 << ClusterIndex );
 
 				if ( ClusterDistributionMask == 0 )
 				{
@@ -241,7 +261,7 @@ void FAutomationControllerManager::ProcessAvailableTasks()
 		}
 	}
 
-	if (bIsLocalSession == false)
+	if ( bIsLocalSession == false )
 	{
 		// Update the test status for timeouts if this is not a local session
 		UpdateTests();
@@ -251,37 +271,31 @@ void FAutomationControllerManager::ProcessAvailableTasks()
 void FAutomationControllerManager::ReportTestResults()
 {
 	GLog->Logf(TEXT("Test Pass Results:"));
-	for (int i = 0; i < OurPassResults.TestInformation.Num(); i++)
+	for ( int32 i = 0; i < OurPassResults.TestInformation.Num(); i++ )
 	{
-		GLog->Logf(TEXT("%s: %s"), *OurPassResults.TestInformation[i].TestName, *OurPassResults.TestInformation[i].TestResult);
+		GLog->Logf(TEXT("%s: %s"), *OurPassResults.TestInformation[i].TestDisplayName, ToString(OurPassResults.TestInformation[i].State));
 	}
 }
 
-void FAutomationControllerManager::CollectTestNotes(FString TestName, const FAutomationWorkerRunTestsReply& Message)
+void FAutomationControllerManager::CollectTestResults(TSharedPtr<IAutomationReport> Report, const FAutomationTestResults& Results)
 {
-	for (int i = 0; i < OurPassResults.TestInformation.Num(); i++)
+	// TODO This is slow, change to a map.
+	for ( int32 i = 0; i < OurPassResults.TestInformation.Num(); i++ )
 	{
-		if (OurPassResults.TestInformation[i].TestName == TestName)
+		FAutomatedTestResult& ReportResult = OurPassResults.TestInformation[i];
+		if ( ReportResult.FullTestPath == Report->GetFullTestPath() )
 		{
-			OurPassResults.TestInformation[i].TestInfo = Message.Logs;
-			OurPassResults.TestInformation[i].TestWarnings = Message.Warnings;
-			for (int j = 0; j < Message.Errors.Num(); j++)
+			ReportResult.Logs = Results.Logs;
+			ReportResult.Warnings = Results.Warnings;
+			for ( int j = 0; j < Results.Errors.Num(); j++ )
 			{
-				OurPassResults.TestInformation[i].TestErrors.Add(Message.Errors[j].Message);
+				ReportResult.Errors.Add(Results.Errors[j].Message);
 			}
-			return;
-		}
-	}
-}
 
-void FAutomationControllerManager::UpdateTestResultValue(FString TestName, EAutomationState::Type TestResult)
-{
-	for (int i = 0; i < OurPassResults.TestInformation.Num(); i++)
-	{
-		if (OurPassResults.TestInformation[i].TestName == TestName)
-		{
-			OurPassResults.TestInformation[i].TestResult = EAutomationState::ToString(TestResult);
-			switch (TestResult)
+			ReportResult.State = Results.State;
+			ReportResult.Artifacts = Results.Artifacts;
+
+			switch ( Results.State )
 			{
 			case EAutomationState::Success:
 				OurPassResults.NumSucceeded++;
@@ -293,6 +307,7 @@ void FAutomationControllerManager::UpdateTestResultValue(FString TestName, EAuto
 				OurPassResults.NumNotRun++;
 				break;
 			}
+
 			return;
 		}
 	}
@@ -304,11 +319,14 @@ void FAutomationControllerManager::GenerateJsonTestPassSummary(FDateTime Timesta
 	{
 		return;
 	}
+
 	const FAutomatedTestPassResults SerializedPassResults = OurPassResults;
 	TSharedPtr<FJsonObject> ReportJson = FJsonObjectConverter::UStructToJsonObject(SerializedPassResults);
 	if (ReportJson.IsValid())
 	{
-		FString ReportFileName = FString::Printf(TEXT("%s/AutomationReport-%d-%s.json"), *FPaths::AutomationLogDir(), FEngineVersion::Current().GetChangelist(), *Timestamp.ToString());
+		FString ReportOutputFolder = ReportOutputPathOverride.IsEmpty() ? FPaths::AutomationLogDir() : ReportOutputPathOverride;
+
+		FString ReportFileName = FString::Printf(TEXT("%s/AutomationReport-%d-%s.json"), *ReportOutputFolder, FEngineVersion::Current().GetChangelist(), *Timestamp.ToString());
 		FArchive* ReportFileWriter = IFileManager::Get().CreateFileWriter(*ReportFileName);
 		if (ReportFileWriter != nullptr)
 		{
@@ -331,77 +349,120 @@ void FAutomationControllerManager::GenerateHtmlTestPassSummary(FDateTime Timesta
 		return;
 	}
 
-	const FAutomatedTestPassResults SerializedPassResults = OurPassResults;
+	FString ReportOutputFolder = ReportOutputPathOverride.IsEmpty() ? FPaths::AutomationLogDir() : ReportOutputPathOverride;
+
+	FAutomatedTestPassResults SerializedPassResults = OurPassResults;
+	SerializedPassResults.TestInformation.StableSort([] (const FAutomatedTestResult& A, const FAutomatedTestResult& B) {
+		if ( A.Errors.Num() > 0 )
+		{
+			if ( B.Errors.Num() > 0 )
+				return ( A.TestDisplayName < B.TestDisplayName );
+			else
+				return true;
+		}
+		else if ( B.Errors.Num() > 0 )
+		{
+			return false;
+		}
+
+		if ( A.Warnings.Num() > 0 )
+		{
+			if ( B.Warnings.Num() > 0 )
+				return ( A.TestDisplayName < B.TestDisplayName );
+			else
+				return true;
+		}
+		else if ( B.Warnings.Num() > 0 )
+		{
+			return false;
+		}
+
+		return A.TestDisplayName < B.TestDisplayName;
+	});
 
-	FString Header =
-		TEXT("<html>")
-		TEXT("<head>")
-		TEXT("<title>Test Results</title>")
-		TEXT("<head>")
-		TEXT("<body>")
-			TEXT("<h3>Test Results</h3>")
-			TEXT("<table>")
-				TEXT("<tr>")
-					TEXT("<th width='10%'>Step</th>")
-					TEXT("<th width='10%'>Result</th>")
-					TEXT("<th width='80%'>Remarks</th>")
-				TEXT("</tr>");
-
-	FString Results;
-	for ( const FAutomatedTestResult& Test : OurPassResults.TestInformation )
+	FString MasterTemplate, ResultTemplate, LogTemplate, ArtifactTemplate;
+	const bool bLoadedMaster = FFileHelper::LoadFileToString(MasterTemplate, *( FPaths::EngineContentDir() / TEXT("Automation/Report-Master-Template.html") ));
+	const bool bLoadedResult = FFileHelper::LoadFileToString(ResultTemplate, *( FPaths::EngineContentDir() / TEXT("Automation/Report-Result-Template.html") ));
+	const bool bLoadedLog = FFileHelper::LoadFileToString(LogTemplate, *( FPaths::EngineContentDir() / TEXT("Automation/Report-Log-Template.html") ));
+	const bool bLoadedArtifact = FFileHelper::LoadFileToString(ArtifactTemplate, *( FPaths::EngineContentDir() / TEXT("Automation/Report-Artifact-Template.html") ));
+	check(bLoadedMaster && bLoadedResult && bLoadedLog && bLoadedArtifact);
+
+	FString SuccessColor = TEXT("#3dd94a");
+	FString WarningColor = TEXT("#f9bb00");
+	FString ErrorColor = TEXT("#ff2154");
+
+	FString TitleColor = SuccessColor;
+	if ( SerializedPassResults.TestInformation.Num() > 0 )
+	{
+		const FAutomatedTestResult& FirstTest = SerializedPassResults.TestInformation[0];
+		TitleColor = FirstTest.Errors.Num() > 0 ? ErrorColor : FirstTest.Warnings.Num() > 0 ? WarningColor : SuccessColor;
+	}
+
+	FString HtmlResults;
+	for ( const FAutomatedTestResult& Test : SerializedPassResults.TestInformation )
 	{
-		FString Row = TEXT("<tr>");
-		Row += FString::Printf(TEXT("<td>%s</td>"), *Test.TestName);
-		Row += FString::Printf(TEXT("<td>%s</td>"), *Test.TestResult);
-		Row += TEXT("<td>");
+		FString TestColor = Test.Errors.Num() > 0 ? ErrorColor : Test.Warnings.Num() > 0 ? WarningColor : SuccessColor;
+
+		FString Logs = TEXT("");
 
-		if ( Test.TestErrors.Num() > 0 )
+		for ( const FString& LogItem : Test.Errors )
 		{
-			Row += TEXT("<h4>Errors</h4>");
-			Row += TEXT("<ul>");
-			for ( const FString& LogItem : Test.TestErrors )
-			{
-				Row += FString::Printf(TEXT("<li>%s</li>"), *FPlatformHttp::HtmlEncode(LogItem));
-			}
-			Row += TEXT("</ul>");
+			TMap<FString, FStringFormatArg> Args;
+			Args.Add(TEXT("Statement"), FPlatformHttp::HtmlEncode(LogItem));
+			Logs += FString::Format(*LogTemplate, Args);
 		}
 
-		if ( Test.TestWarnings.Num() > 0 )
+		for ( const FString& LogItem : Test.Warnings )
 		{
-			Row += TEXT("<h4>Warnings</h4>");
-			Row += TEXT("<ul>");
-			for ( const FString& LogItem : Test.TestWarnings )
-			{
-				Row += FString::Printf(TEXT("<li>%s</li>"), *FPlatformHttp::HtmlEncode(LogItem));
-			}
-			Row += TEXT("</ul>");
+			TMap<FString, FStringFormatArg> Args;
+			Args.Add(TEXT("Statement"), FPlatformHttp::HtmlEncode(LogItem));
+			Logs += FString::Format(*LogTemplate, Args);
 		}
 
-		if ( Test.TestInfo.Num() > 0 )
+		for ( const FString& LogItem : Test.Logs )
 		{
-			Row += TEXT("<h4>Info</h4>");
-			Row += TEXT("<ul>");
-			for ( const FString& LogItem : Test.TestInfo )
-			{
-				Row += FString::Printf(TEXT("<li>%s</li>"), *FPlatformHttp::HtmlEncode(LogItem));
-			}
-			Row += TEXT("</ul>");
+			TMap<FString, FStringFormatArg> Args;
+			Args.Add(TEXT("Statement"), FPlatformHttp::HtmlEncode(LogItem));
+			Logs += FString::Format(*LogTemplate, Args);
 		}
 
-		Row += TEXT("</td>");
-		Row += TEXT("</tr>");
+		FString AritfactDirectory = FString::Printf(TEXT("ReportArtifact-%d"), +FEngineVersion::Current().GetChangelist());
+		for ( const FAutomationArtifact& Artifact : Test.Artifacts )
+		{
+			FString ArtifactFile = AritfactDirectory / FGuid::NewGuid().ToString(EGuidFormats::Digits) + FPaths::GetExtension(Artifact.FilePath, true);
+			FString ArtifactDestination = ReportOutputFolder / ArtifactFile;
+			IFileManager::Get().Copy(*ArtifactDestination, *Artifact.FilePath, true, true);
+
+			TMap<FString, FStringFormatArg> Args;
+			Args.Add(TEXT("Name"), FPlatformHttp::HtmlEncode(Artifact.Name));
+			Args.Add(TEXT("File"), ArtifactFile);
+			Logs += FString::Format(*ArtifactTemplate, Args);
+		}
 
-		Results += Row;
-	}
+		{
+			TMap<FString, FStringFormatArg> Args;
+			Args.Add(TEXT("TestColor"), TestColor);
+			Args.Add(TEXT("TestName"), Test.TestDisplayName);
+			Args.Add(TEXT("TestPath"), Test.FullTestPath);
+			Args.Add(TEXT("Logs"), Logs);
 
-	FString Footer =
-			TEXT("</table>")
-		TEXT("</body>");
+			HtmlResults += FString::Format(*ResultTemplate, Args);
+		}
+	}
 
-	FString ReportFileName = FString::Printf(TEXT("%s/AutomationReport-%d-%s.html"), *FPaths::AutomationLogDir(), FEngineVersion::Current().GetChangelist(), *Timestamp.ToString());
-	if ( !FFileHelper::SaveStringToFile(Header + Results + Footer, *ReportFileName) )
 	{
-		GLog->Logf(ELogVerbosity::Error, TEXT("Test Report Html is invalid - report not generated."));
+		TMap<FString, FStringFormatArg> Args;
+		Args.Add(TEXT("Title"), TEXT("Automation Test Results"));
+		Args.Add(TEXT("TitleColor"), TitleColor);
+		Args.Add(TEXT("Results"), HtmlResults);
+
+		FString Html = FString::Format(*MasterTemplate, Args);
+
+		FString ReportFileName = FString::Printf(TEXT("%s/AutomationReport-%d-%s.html"), *ReportOutputFolder, FEngineVersion::Current().GetChangelist(), *Timestamp.ToString());
+		if ( !FFileHelper::SaveStringToFile(Html, *ReportFileName) )
+		{
+			GLog->Logf(ELogVerbosity::Error, TEXT("Test Report Html is invalid - report not generated."));
+		}
 	}
 }
 
@@ -412,42 +473,45 @@ void FAutomationControllerManager::ExecuteNextTask( int32 ClusterIndex, OUT bool
 
 	// For each device in this cluster
 	int32 NumDevicesInCluster = DeviceClusterManager.GetNumDevicesInCluster( ClusterIndex );
-	for( int32 DeviceIndex = 0; DeviceIndex < NumDevicesInCluster; ++DeviceIndex )
+	for ( int32 DeviceIndex = 0; DeviceIndex < NumDevicesInCluster; ++DeviceIndex )
 	{
 		// If this device is idle
-		if( !DeviceClusterManager.GetTest( ClusterIndex, DeviceIndex ).IsValid() && DeviceClusterManager.DeviceEnabled( ClusterIndex, DeviceIndex ) )
+		if ( !DeviceClusterManager.GetTest(ClusterIndex, DeviceIndex).IsValid() && DeviceClusterManager.DeviceEnabled(ClusterIndex, DeviceIndex) )
 		{
 			// Get the next test that should be worked on
-			TSharedPtr< IAutomationReport > NextTest = ReportManager.GetNextReportToExecute( bAllTestsCompleted, ClusterIndex,CurrentTestPass, NumDevicesInCluster );
-			if( NextTest.IsValid() )
+			TSharedPtr< IAutomationReport > NextTest = ReportManager.GetNextReportToExecute(bAllTestsCompleted, ClusterIndex, CurrentTestPass, NumDevicesInCluster);
+			if ( NextTest.IsValid() )
 			{
 				// Get the status of the test
-				EAutomationState::Type TestState = NextTest->GetState( ClusterIndex,CurrentTestPass );
-				if( TestState == EAutomationState::NotRun )
+				EAutomationState TestState = NextTest->GetState(ClusterIndex, CurrentTestPass);
+				if ( TestState == EAutomationState::NotRun )
 				{
 					// Reserve this device for the test
-					DeviceClusterManager.SetTest( ClusterIndex, DeviceIndex, NextTest );
-					TestsRunThisPass.Add( NextTest );
+					DeviceClusterManager.SetTest(ClusterIndex, DeviceIndex, NextTest);
+					TestsRunThisPass.Add(NextTest);
 
 					// Register this as a test we'll need to report on.
 					FAutomatedTestResult tempresult;
-					tempresult.TestName = NextTest->GetDisplayName();
+					tempresult.Test = NextTest;
+					tempresult.TestDisplayName = NextTest->GetDisplayName();
+					tempresult.FullTestPath = NextTest->GetFullTestPath();
+
 					OurPassResults.TestInformation.Add(tempresult);
 
 					// If we now have enough devices reserved for the test, run it!
-					TArray<FMessageAddress> DeviceAddresses = DeviceClusterManager.GetDevicesReservedForTest( ClusterIndex, NextTest );
-					if( DeviceAddresses.Num() == NextTest->GetNumParticipantsRequired() )
+					TArray<FMessageAddress> DeviceAddresses = DeviceClusterManager.GetDevicesReservedForTest(ClusterIndex, NextTest);
+					if ( DeviceAddresses.Num() == NextTest->GetNumParticipantsRequired() )
 					{
 						// Send it to each device
-						for (int32 AddressIndex = 0; AddressIndex < DeviceAddresses.Num(); ++AddressIndex)
+						for ( int32 AddressIndex = 0; AddressIndex < DeviceAddresses.Num(); ++AddressIndex )
 						{
 							FAutomationTestResults TestResults;
 
 							GLog->Logf(ELogVerbosity::Display, TEXT("Running Automation: '%s' (Class Name: '%s')"), *TestsRunThisPass[AddressIndex]->GetFullTestPath(), *TestsRunThisPass[AddressIndex]->GetCommand());
 							TestResults.State = EAutomationState::InProcess;
 
-							TestResults.GameInstance = DeviceClusterManager.GetClusterDeviceName( ClusterIndex, DeviceIndex );
-							NextTest->SetResults( ClusterIndex,CurrentTestPass, TestResults );
+							TestResults.GameInstance = DeviceClusterManager.GetClusterDeviceName(ClusterIndex, DeviceIndex);
+							NextTest->SetResults(ClusterIndex, CurrentTestPass, TestResults);
 							NextTest->ResetNetworkCommandResponses();
 
 							// Mark the device as busy
@@ -457,7 +521,7 @@ void FAutomationControllerManager::ExecuteNextTask( int32 ClusterIndex, OUT bool
 							MessageEndpoint->Send(new FAutomationWorkerRunTests(ExecutionCount, AddressIndex, NextTest->GetCommand(), NextTest->GetDisplayName(), bScreenshotsEnabled, bSendAnalytics), DeviceAddress);
 
 							// Add a test so we can check later if the device is still active
-							TestRunningArray.Add( FTestRunningInfo( DeviceAddress ) );
+							TestRunningArray.Add(FTestRunningInfo(DeviceAddress));
 						}
 					}
 				}
@@ -471,32 +535,32 @@ void FAutomationControllerManager::ExecuteNextTask( int32 ClusterIndex, OUT bool
 	}
 
 	// Ensure any tests we have attempted to run on this pass had enough participants to successfully run.
-	for( int32 TestIndex = 0; TestIndex < TestsRunThisPass.Num(); TestIndex++ )
+	for ( int32 TestIndex = 0; TestIndex < TestsRunThisPass.Num(); TestIndex++ )
 	{
-		IAutomationReportPtr CurrentTest = TestsRunThisPass[ TestIndex ];
+		IAutomationReportPtr CurrentTest = TestsRunThisPass[TestIndex];
 
-		if( CurrentTest->GetNumDevicesRunningTest() != CurrentTest->GetNumParticipantsRequired() )
+		if ( CurrentTest->GetNumDevicesRunningTest() != CurrentTest->GetNumParticipantsRequired() )
 		{
-			if( GetNumDevicesInCluster( ClusterIndex ) < CurrentTest->GetNumParticipantsRequired()  )
+			if ( GetNumDevicesInCluster(ClusterIndex) < CurrentTest->GetNumParticipantsRequired() )
 			{
 				float EmptyDuration = 0.0f;
 				TArray<FString> EmptyStringArray;
 				TArray<FString> AutomationsWarnings;
-				AutomationsWarnings.Add( FString::Printf( TEXT( "Needed %d devices to participate, Only had %d available." ), CurrentTest->GetNumParticipantsRequired(), DeviceClusterManager.GetNumDevicesInCluster( ClusterIndex ) ) );
+				AutomationsWarnings.Add(FString::Printf(TEXT("Needed %d devices to participate, Only had %d available."), CurrentTest->GetNumParticipantsRequired(), DeviceClusterManager.GetNumDevicesInCluster(ClusterIndex)));
 
 				FAutomationTestResults TestResults;
 				TestResults.State = EAutomationState::NotEnoughParticipants;
 				TestResults.GameInstance = DeviceClusterManager.GetClusterDeviceName(ClusterIndex, 0);
-				TestResults.Warnings.Append( AutomationsWarnings );
+				TestResults.Warnings.Append(AutomationsWarnings);
 
-				CurrentTest->SetResults( ClusterIndex,CurrentTestPass, TestResults );
-				DeviceClusterManager.ResetAllDevicesRunningTest( ClusterIndex, CurrentTest );
+				CurrentTest->SetResults(ClusterIndex, CurrentTestPass, TestResults);
+				DeviceClusterManager.ResetAllDevicesRunningTest(ClusterIndex, CurrentTest);
 			}
 		}
 	}
 
 	//Check to see if we finished a pass
-	if( bAllTestsCompleted && CurrentTestPass < NumTestPasses - 1 )
+	if ( bAllTestsCompleted && CurrentTestPass < NumTestPasses - 1 )
 	{
 		CurrentTestPass++;
 		ReportManager.SetCurrentTestPass(CurrentTestPass);
@@ -517,7 +581,7 @@ void FAutomationControllerManager::Startup()
 		.Handling<FAutomationWorkerScreenImage>(this, &FAutomationControllerManager::HandleReceivedScreenShot)
 		.Handling<FAutomationWorkerWorkerOffline>(this, &FAutomationControllerManager::HandleWorkerOfflineMessage);
 
-	if (MessageEndpoint.IsValid())
+	if ( MessageEndpoint.IsValid() )
 	{
 		MessageEndpoint->Subscribe<FAutomationWorkerWorkerOffline>();
 	}
@@ -535,7 +599,6 @@ void FAutomationControllerManager::Startup()
 	ToggleDeviceGroupFlag(EAutomationDeviceGroupTypes::MachineName);
 }
 
-
 void FAutomationControllerManager::Shutdown()
 {
 	MessageEndpoint.Reset();
@@ -543,7 +606,6 @@ void FAutomationControllerManager::Shutdown()
 	RemoveCallbacks();
 }
 
-
 void FAutomationControllerManager::RemoveCallbacks()
 {
 	ShutdownDelegate.Clear();
@@ -552,19 +614,18 @@ void FAutomationControllerManager::RemoveCallbacks()
 	TestsCompleteDelegate.Clear();
 }
 
-
-void FAutomationControllerManager::SetTestNames( const FMessageAddress& AutomationWorkerAddress )
+void FAutomationControllerManager::SetTestNames(const FMessageAddress& AutomationWorkerAddress)
 {
 	int32 DeviceClusterIndex = INDEX_NONE;
 	int32 DeviceIndex = INDEX_NONE;
 
 	// Find the device that requested these tests
-	if( DeviceClusterManager.FindDevice( AutomationWorkerAddress, DeviceClusterIndex, DeviceIndex ) )
+	if ( DeviceClusterManager.FindDevice(AutomationWorkerAddress, DeviceClusterIndex, DeviceIndex) )
 	{
 		// Sort tests by display name
 		struct FCompareAutomationTestInfo
 		{
-			FORCEINLINE bool operator()(const FAutomationTestInfo& A,const FAutomationTestInfo& B) const
+			FORCEINLINE bool operator()(const FAutomationTestInfo& A, const FAutomationTestInfo& B) const
 			{
 				return A.GetDisplayName() < B.GetDisplayName();
 			}
@@ -573,10 +634,10 @@ void FAutomationControllerManager::SetTestNames( const FMessageAddress& Automati
 		TestInfo.Sort(FCompareAutomationTestInfo());
 
 		// Add each test to the collection
-		for( int32 TestIndex = 0; TestIndex < TestInfo.Num(); ++TestIndex )
+		for ( int32 TestIndex = 0; TestIndex < TestInfo.Num(); ++TestIndex )
 		{
 			// Ensure our test exists. If not, add it
-			ReportManager.EnsureReportExists( TestInfo[TestIndex], DeviceClusterIndex, NumTestPasses );
+			ReportManager.EnsureReportExists(TestInfo[TestIndex], DeviceClusterIndex, NumTestPasses);
 		}
 
 		// Clear any intermediate data we had associated with the tests whilst building the full list of tests
@@ -591,7 +652,7 @@ void FAutomationControllerManager::SetTestNames( const FMessageAddress& Automati
 	RefreshTestResponses++;
 
 	// If we have received all the responses we expect to
-	if (RefreshTestResponses == DeviceClusterManager.GetNumClusters())
+	if ( RefreshTestResponses == DeviceClusterManager.GetNumClusters() )
 	{
 		TestsRefreshedDelegate.Broadcast();
 
@@ -600,7 +661,6 @@ void FAutomationControllerManager::SetTestNames( const FMessageAddress& Automati
 	}
 }
 
-
 void FAutomationControllerManager::ProcessResults()
 {
 	bHasErrors = false;
@@ -609,13 +669,13 @@ void FAutomationControllerManager::ProcessResults()
 
 	TArray< TSharedPtr< IAutomationReport > >& TestReports = GetReports();
 
-	if (TestReports.Num())
+	if ( TestReports.Num() )
 	{
 		bTestResultsAvailable = true;
 
-		for (int32 Index = 0; Index < TestReports.Num(); Index++)
+		for ( int32 Index = 0; Index < TestReports.Num(); Index++ )
 		{
-			CheckChildResult(TestReports[Index ]);
+			CheckChildResult(TestReports[Index]);
 		}
 	}
 
@@ -626,36 +686,35 @@ void FAutomationControllerManager::ProcessResults()
 
 	OurPassResults.ClearAllEntries();
 
-	SetControllerStatus( EAutomationControllerModuleState::Ready );
+	SetControllerStatus(EAutomationControllerModuleState::Ready);
 }
 
-
 void FAutomationControllerManager::CheckChildResult(TSharedPtr<IAutomationReport> InReport)
 {
 	TArray<TSharedPtr<IAutomationReport> >& ChildReports = InReport->GetChildReports();
 
-	if (ChildReports.Num() > 0)
+	if ( ChildReports.Num() > 0 )
 	{
-		for (int32 Index = 0; Index < ChildReports.Num(); Index++)
+		for ( int32 Index = 0; Index < ChildReports.Num(); Index++ )
 		{
 			CheckChildResult(ChildReports[Index]);
 		}
 	}
-	else if ((bHasErrors && bHasWarning && bHasLogs ) == false && InReport->IsEnabled())
+	else if ( ( bHasErrors && bHasWarning && bHasLogs ) == false && InReport->IsEnabled() )
 	{
-		for (int32 ClusterIndex = 0; ClusterIndex < GetNumDeviceClusters(); ++ClusterIndex)
+		for ( int32 ClusterIndex = 0; ClusterIndex < GetNumDeviceClusters(); ++ClusterIndex )
 		{
-			FAutomationTestResults TestResults = InReport->GetResults( ClusterIndex,CurrentTestPass );
+			FAutomationTestResults TestResults = InReport->GetResults(ClusterIndex, CurrentTestPass);
 
-			if (TestResults.Errors.Num())
+			if ( TestResults.Errors.Num() )
 			{
 				bHasErrors = true;
 			}
-			if (TestResults.Warnings.Num())
+			if ( TestResults.Warnings.Num() )
 			{
 				bHasWarning = true;
 			}
-			if (TestResults.Logs.Num())
+			if ( TestResults.Logs.Num() )
 			{
 				bHasLogs = true;
 			}
@@ -663,10 +722,9 @@ void FAutomationControllerManager::CheckChildResult(TSharedPtr<IAutomationReport
 	}
 }
 
-
-void FAutomationControllerManager::SetControllerStatus( EAutomationControllerModuleState::Type InAutomationTestState )
+void FAutomationControllerManager::SetControllerStatus(EAutomationControllerModuleState::Type InAutomationTestState)
 {
-	if (InAutomationTestState != AutomationTestState)
+	if ( InAutomationTestState != AutomationTestState )
 	{
 		// Inform the UI if the test state has changed
 		AutomationTestState = InAutomationTestState;
@@ -674,34 +732,31 @@ void FAutomationControllerManager::SetControllerStatus( EAutomationControllerMod
 	}
 }
 
-
-void FAutomationControllerManager::RemoveTestRunning( const FMessageAddress& TestAddressToRemove )
+void FAutomationControllerManager::RemoveTestRunning(const FMessageAddress& TestAddressToRemove)
 {
 	for ( int32 Index = 0; Index < TestRunningArray.Num(); Index++ )
 	{
-		if ( TestRunningArray[ Index ].OwnerMessageAddress == TestAddressToRemove )
+		if ( TestRunningArray[Index].OwnerMessageAddress == TestAddressToRemove )
 		{
-			TestRunningArray.RemoveAt( Index );
+			TestRunningArray.RemoveAt(Index);
 			break;
 		}
 	}
 }
 
-
-void FAutomationControllerManager::AddPingResult( const FMessageAddress& ResponderAddress )
+void FAutomationControllerManager::AddPingResult(const FMessageAddress& ResponderAddress)
 {
 	for ( int32 Index = 0; Index < TestRunningArray.Num(); Index++ )
 	{
-		if ( TestRunningArray[ Index ].OwnerMessageAddress == ResponderAddress )
+		if ( TestRunningArray[Index].OwnerMessageAddress == ResponderAddress )
 		{
-			TestRunningArray[ Index ].LastPingTime = 0;
+			TestRunningArray[Index].LastPingTime = 0;
 			break;
 		}
 	}
 }
 
-
-void FAutomationControllerManager::UpdateTests( )
+void FAutomationControllerManager::UpdateTests()
 {
 	static const float CheckTestInterval = 1.0f;
 	static const float GameInstanceLostTimer = 200.0f;
@@ -712,22 +767,22 @@ void FAutomationControllerManager::UpdateTests( )
 	{
 		for ( int32 Index = 0; Index < TestRunningArray.Num(); Index++ )
 		{
-			TestRunningArray[ Index ].LastPingTime += CheckTestTimer;
+			TestRunningArray[Index].LastPingTime += CheckTestTimer;
 
-			if ( TestRunningArray[ Index ].LastPingTime > GameInstanceLostTimer )
+			if ( TestRunningArray[Index].LastPingTime > GameInstanceLostTimer )
 			{
 				// Find the game session instance info
 				int32 ClusterIndex;
 				int32 DeviceIndex;
-				verify( DeviceClusterManager.FindDevice( TestRunningArray[ Index ].OwnerMessageAddress, ClusterIndex, DeviceIndex ) );
+				verify(DeviceClusterManager.FindDevice(TestRunningArray[Index].OwnerMessageAddress, ClusterIndex, DeviceIndex));
 				//verify this device thought it was busy
 				TSharedPtr <IAutomationReport> Report = DeviceClusterManager.GetTest(ClusterIndex, DeviceIndex);
-				check (Report.IsValid());
+				check(Report.IsValid());
 				// A dummy array used to report the result
 
 				TArray<FString> EmptyStringArray;
 				TArray<FString> ErrorStringArray;
-				ErrorStringArray.Add( FString( TEXT( "Failed" ) ) );
+				ErrorStringArray.Add(FString(TEXT("Failed")));
 				bHasErrors = true;
 				GLog->Logf(ELogVerbosity::Display, TEXT("Timeout hit. Nooooooo."));
 
@@ -736,33 +791,33 @@ void FAutomationControllerManager::UpdateTests( )
 				TestResults.GameInstance = DeviceClusterManager.GetClusterDeviceName(ClusterIndex, DeviceIndex);
 
 				// Set the results
-				Report->SetResults(ClusterIndex,CurrentTestPass, TestResults );
+				Report->SetResults(ClusterIndex, CurrentTestPass, TestResults);
 				bTestResultsAvailable = true;
 
 				// Disable the device in the cluster so it is not used again
-				DeviceClusterManager.DisableDevice( ClusterIndex, DeviceIndex );
+				DeviceClusterManager.DisableDevice(ClusterIndex, DeviceIndex);
 
 				// Remove the running test
-				TestRunningArray.RemoveAt( Index-- );
+				TestRunningArray.RemoveAt(Index--);
 
 				// If there are no more devices, set the module state to disabled
 				if ( DeviceClusterManager.HasActiveDevice() == false )
 				{
 					GLog->Logf(ELogVerbosity::Display, TEXT("Module disabled"));
-					SetControllerStatus( EAutomationControllerModuleState::Disabled );
+					SetControllerStatus(EAutomationControllerModuleState::Disabled);
 					ClusterDistributionMask = 0;
 				}
 				else
 				{
 					GLog->Logf(ELogVerbosity::Display, TEXT("Module not disabled. Keep looking."));
 					// Remove the cluster from the mask if there are no active devices left
-					if ( DeviceClusterManager.GetNumActiveDevicesInCluster( ClusterIndex ) == 0 )
+					if ( DeviceClusterManager.GetNumActiveDevicesInCluster(ClusterIndex) == 0 )
 					{
-						ClusterDistributionMask &= ~( 1<<ClusterIndex );
+						ClusterDistributionMask &= ~( 1 << ClusterIndex );
 					}
-					if ( TestRunningArray.Num() == 0  )
+					if ( TestRunningArray.Num() == 0 )
 					{
-						SetControllerStatus( EAutomationControllerModuleState::Ready );
+						SetControllerStatus(EAutomationControllerModuleState::Ready);
 					}
 				}
 			}
@@ -775,22 +830,20 @@ void FAutomationControllerManager::UpdateTests( )
 	}
 }
 
-
-const bool FAutomationControllerManager::ExportReport( uint32 FileExportTypeMask )
+const bool FAutomationControllerManager::ExportReport(uint32 FileExportTypeMask)
 {
-	return ReportManager.ExportReport( FileExportTypeMask, GetNumDeviceClusters() );
+	return ReportManager.ExportReport(FileExportTypeMask, GetNumDeviceClusters());
 }
 
-
-bool FAutomationControllerManager::IsTestRunnable( IAutomationReportPtr InReport ) const
+bool FAutomationControllerManager::IsTestRunnable(IAutomationReportPtr InReport) const
 {
 	bool bIsRunnable = false;
 
-	for( int32 ClusterIndex = 0; ClusterIndex < GetNumDeviceClusters(); ++ClusterIndex )
+	for ( int32 ClusterIndex = 0; ClusterIndex < GetNumDeviceClusters(); ++ClusterIndex )
 	{
-		if( InReport->IsSupported( ClusterIndex ) )
+		if ( InReport->IsSupported(ClusterIndex) )
 		{
-			if( GetNumDevicesInCluster( ClusterIndex ) >= InReport->GetNumParticipantsRequired() )
+			if ( GetNumDevicesInCluster(ClusterIndex) >= InReport->GetNumParticipantsRequired() )
 			{
 				bIsRunnable = true;
 				break;
@@ -804,25 +857,24 @@ bool FAutomationControllerManager::IsTestRunnable( IAutomationReportPtr InReport
 /* FAutomationControllerModule callbacks
  *****************************************************************************/
 
-void FAutomationControllerManager::HandleFindWorkersResponseMessage( const FAutomationWorkerFindWorkersResponse& Message, const IMessageContextRef& Context )
+void FAutomationControllerManager::HandleFindWorkersResponseMessage(const FAutomationWorkerFindWorkersResponse& Message, const IMessageContextRef& Context)
 {
-	if (Message.SessionId == ActiveSessionId)
+	if ( Message.SessionId == ActiveSessionId )
 	{
 		DeviceClusterManager.AddDeviceFromMessage(Context->GetSender(), Message, DeviceGroupFlags);
 	}
 
 	RequestTests();
 
-	SetControllerStatus( EAutomationControllerModuleState::Ready );
+	SetControllerStatus(EAutomationControllerModuleState::Ready);
 }
 
-
 void FAutomationControllerManager::HandlePongMessage( const FAutomationWorkerPong& Message, const IMessageContextRef& Context )
 {
 	AddPingResult(Context->GetSender());
 }
 
-void FAutomationControllerManager::HandleReceivedScreenShot( const FAutomationWorkerScreenImage& Message, const IMessageContextRef& Context )
+void FAutomationControllerManager::HandleReceivedScreenShot(const FAutomationWorkerScreenImage& Message, const IMessageContextRef& Context)
 {
 	FString ScreenshotIncomingFolder = FPaths::GameSavedDir() / TEXT("Automation/Incoming/");
 
@@ -848,10 +900,10 @@ void FAutomationControllerManager::HandleReceivedScreenShot( const FAutomationWo
 	ComparisonQueue.Enqueue(Comparison);
 }
 
-void FAutomationControllerManager::HandleRequestNextNetworkCommandMessage( const FAutomationWorkerRequestNextNetworkCommand& Message, const IMessageContextRef& Context )
+void FAutomationControllerManager::HandleRequestNextNetworkCommandMessage(const FAutomationWorkerRequestNextNetworkCommand& Message, const IMessageContextRef& Context)
 {
 	// Harvest iteration of running the tests this result came from (stops stale results from being committed to subsequent runs)
-	if (Message.ExecutionCount == ExecutionCount)
+	if ( Message.ExecutionCount == ExecutionCount )
 	{
 		// Find the device id for the address
 		int32 ClusterIndex;
@@ -860,24 +912,24 @@ void FAutomationControllerManager::HandleRequestNextNetworkCommandMessage( const
 		verify(DeviceClusterManager.FindDevice(Context->GetSender(), ClusterIndex, DeviceIndex));
 
 		// Verify this device thought it was busy
-		TSharedPtr <IAutomationReport> Report = DeviceClusterManager.GetTest(ClusterIndex, DeviceIndex);
-		check (Report.IsValid());
+		TSharedPtr<IAutomationReport> Report = DeviceClusterManager.GetTest(ClusterIndex, DeviceIndex);
+		check(Report.IsValid());
 
 		// Increment network command responses
 		bool bAllResponsesReceived = Report->IncrementNetworkCommandResponses();
 
 		// Test if we've accumulated all responses AND this was the result for the round of test running AND we're still running tests
-		if (bAllResponsesReceived && (ClusterDistributionMask & (1<<ClusterIndex))) 
+		if ( bAllResponsesReceived && ( ClusterDistributionMask & ( 1 << ClusterIndex ) ) )
 		{
 			// Reset the counter
 			Report->ResetNetworkCommandResponses();
 
 			// For every device in this networked test
 			TArray<FMessageAddress> DeviceAddresses = DeviceClusterManager.GetDevicesReservedForTest(ClusterIndex, Report);
-			check (DeviceAddresses.Num() == Report->GetNumParticipantsRequired());
+			check(DeviceAddresses.Num() == Report->GetNumParticipantsRequired());
 
 			// Send it to each device
-			for (int32 AddressIndex = 0; AddressIndex < DeviceAddresses.Num(); ++AddressIndex)
+			for ( int32 AddressIndex = 0; AddressIndex < DeviceAddresses.Num(); ++AddressIndex )
 			{
 				//send "next command message" to worker
 				MessageEndpoint->Send(new FAutomationWorkerNextNetworkCommandReply(), DeviceAddresses[AddressIndex]);
@@ -886,24 +938,21 @@ void FAutomationControllerManager::HandleRequestNextNetworkCommandMessage( const
 	}
 }
 
-
 void FAutomationControllerManager::HandleRequestTestsReplyMessage(const FAutomationWorkerRequestTestsReply& Message, const IMessageContextRef& Context)
 {
 	FAutomationTestInfo NewTest = Message.GetTestInfo();
 	TestInfo.Add(NewTest);
 }
 
-
 void FAutomationControllerManager::HandleRequestTestsReplyCompleteMessage(const FAutomationWorkerRequestTestsReplyComplete& Message, const IMessageContextRef& Context)
 {
 	SetTestNames(Context->GetSender());
 }
 
-
-void FAutomationControllerManager::HandleRunTestsReplyMessage( const FAutomationWorkerRunTestsReply& Message, const IMessageContextRef& Context )
+void FAutomationControllerManager::HandleRunTestsReplyMessage(const FAutomationWorkerRunTestsReply& Message, const IMessageContextRef& Context)
 {
 	// If we should commit these results
-	if (Message.ExecutionCount == ExecutionCount)
+	if ( Message.ExecutionCount == ExecutionCount )
 	{
 		FAutomationTestResults TestResults;
 
@@ -930,9 +979,10 @@ void FAutomationControllerManager::HandleRunTestsReplyMessage( const FAutomation
 
 		Report->SetResults(ClusterIndex, CurrentTestPass, TestResults);
 
+		const FAutomationTestResults& FinalResults = Report->GetResults(ClusterIndex, CurrentTestPass);
+
 		// Gather all of the data relevant to this test for our json reporting.
-		CollectTestNotes(Report->GetDisplayName(), Message);
-		UpdateTestResultValue(Report->GetDisplayName(), TestResults.State);
+		CollectTestResults(Report, FinalResults);
 
 #if WITH_EDITOR
 		FMessageLog AutomationTestingLog("AutomationTestingLog");
@@ -947,14 +997,14 @@ void FAutomationControllerManager::HandleRunTestsReplyMessage( const FAutomation
 			AutomationTestingLog.Error(FText::FromString(( *ErrorIter ).ToString()));
 #endif
 		}
-		for (TArray<FString>::TConstIterator WarningIter(Message.Warnings); WarningIter; ++WarningIter)
+		for ( TArray<FString>::TConstIterator WarningIter(Message.Warnings); WarningIter; ++WarningIter )
 		{
 			GLog->Logf(ELogVerbosity::Warning, TEXT("%s"), **WarningIter);
 #if WITH_EDITOR
 			AutomationTestingLog.Warning(FText::FromString(*WarningIter));
 #endif
 		}
-		for (TArray<FString>::TConstIterator LogItemIter(Message.Logs); LogItemIter; ++LogItemIter)
+		for ( TArray<FString>::TConstIterator LogItemIter(Message.Logs); LogItemIter; ++LogItemIter )
 		{
 			GLog->Logf(ELogVerbosity::Log, TEXT("%s"), **LogItemIter);
 #if WITH_EDITOR
@@ -962,7 +1012,7 @@ void FAutomationControllerManager::HandleRunTestsReplyMessage( const FAutomation
 #endif
 		}
 
-		if (TestResults.State == EAutomationState::Success)
+		if ( TestResults.State == EAutomationState::Success )
 		{
 			FString SuccessString = FString::Printf(TEXT("...Automation Test Succeeded (%s)"), *Report->GetDisplayName());
 			GLog->Logf(ELogVerbosity::Log, *SuccessString);
@@ -991,7 +1041,6 @@ void FAutomationControllerManager::HandleRunTestsReplyMessage( const FAutomation
 	RemoveTestRunning(Context->GetSender());
 }
 
-
 void FAutomationControllerManager::HandleWorkerOfflineMessage( const FAutomationWorkerWorkerOffline& Message, const IMessageContextRef& Context )
 {
 	FMessageAddress DeviceMessageAddress = Context->GetSender();
diff --git a/Engine/Source/Developer/AutomationController/Private/AutomationReport.cpp b/Engine/Source/Developer/AutomationController/Private/AutomationReport.cpp
--- a/Engine/Source/Developer/AutomationController/Private/AutomationReport.cpp
+++ b/Engine/Source/Developer/AutomationController/Private/AutomationReport.cpp
@@ -622,7 +622,10 @@ void FAutomationReport::SetResults( const int32 ClusterIndex, const int32 PassIn
 		TestInfo.InformOfNewDeviceRunningTest();
 	}
 
-	Results[ ClusterIndex ][ PassIndex ] = InResults;
+	TArray<FAutomationArtifact> ExistingArtifacts = Results[ClusterIndex][PassIndex].Artifacts;
+	Results[ClusterIndex][PassIndex] = InResults;
+	Results[ClusterIndex][PassIndex].Artifacts.Append(ExistingArtifacts);
+
 	// Add an error report if none was received
 	if ( InResults.State == EAutomationState::Fail && InResults.Errors.Num() == 0 && InResults.Warnings.Num() == 0 )
 	{
@@ -641,13 +644,21 @@ void FAutomationReport::SetResults( const int32 ClusterIndex, const int32 PassIn
 	OnSetResults.ExecuteIfBound(AsShared());
 }
 
+void FAutomationReport::AddArtifact(const int32 ClusterIndex, const int32 PassIndex, const FAutomationArtifact& Artifact)
+{
+	//verify this is a platform this test is aware of
+	check(( ClusterIndex >= 0 ) && ( ClusterIndex < Results.Num() ));
+	check(( PassIndex >= 0 ) && ( PassIndex < Results[ClusterIndex].Num() ));
+
+	Results[ClusterIndex][PassIndex].Artifacts.Add(Artifact);
+}
 
 void FAutomationReport::GetCompletionStatus(const int32 ClusterIndex, const int32 PassIndex, FAutomationCompleteState& OutCompletionState)
 {
 	//if this test is enabled and a leaf test
 	if (IsSupported(ClusterIndex) && (ChildReports.Num()==0))
 	{
-		EAutomationState::Type CurrentState = Results[ClusterIndex][PassIndex].State;
+		EAutomationState CurrentState = Results[ClusterIndex][PassIndex].State;
 		//Enabled and In-Process
 		if (IsEnabled())
 		{
@@ -686,7 +697,7 @@ void FAutomationReport::GetCompletionStatus(const int32 ClusterIndex, const int3
 }
 
 
-EAutomationState::Type FAutomationReport::GetState(const int32 ClusterIndex, const int32 PassIndex) const
+EAutomationState FAutomationReport::GetState(const int32 ClusterIndex, const int32 PassIndex) const
 {
 	if ((ClusterIndex >= 0) && (ClusterIndex < Results.Num()) &&
 		(PassIndex >= 0) && (PassIndex < Results[ClusterIndex].Num()))
@@ -846,7 +857,7 @@ TSharedPtr<IAutomationReport> FAutomationReport::GetNextReportToExecute(bool& bO
 		//consider self
 		if (IsEnabled() && IsSupported(ClusterIndex))
 		{
-			EAutomationState::Type TestState = GetState(ClusterIndex,PassIndex);
+			EAutomationState TestState = GetState(ClusterIndex,PassIndex);
 			//if any enabled test hasn't been run yet or is in process
 			if ((TestState != EAutomationState::Success) && (TestState != EAutomationState::Fail) && (TestState != EAutomationState::NotEnoughParticipants))
 			{
diff --git a/Engine/Source/Developer/AutomationController/Private/AutomationReport.h b/Engine/Source/Developer/AutomationController/Private/AutomationReport.h
--- a/Engine/Source/Developer/AutomationController/Private/AutomationReport.h
+++ b/Engine/Source/Developer/AutomationController/Private/AutomationReport.h
@@ -47,9 +47,10 @@ public:
 	virtual TArray<TSharedPtr<IAutomationReport> >& GetChildReports() override;
 	virtual void ClustersUpdated(const int32 NumClusters) override;
 	virtual void ResetForExecution(const int32 NumTestPasses) override;
-	virtual void SetResults( const int32 ClusterIndex, const int32 PassIndex, const FAutomationTestResults& InResults ) override;
+	virtual void SetResults(const int32 ClusterIndex, const int32 PassIndex, const FAutomationTestResults& InResults) override;
+	virtual void AddArtifact(const int32 ClusterIndex, const int32 PassIndex, const FAutomationArtifact& Artifact) override;
 	virtual void GetCompletionStatus(const int32 ClusterIndex, const int32 PassIndex, FAutomationCompleteState& OutCompletionState) override;
-	virtual EAutomationState::Type GetState(const int32 ClusterIndex, const int32 PassIndex) const override;
+	virtual EAutomationState GetState(const int32 ClusterIndex, const int32 PassIndex) const override;
 	virtual const FAutomationTestResults& GetResults( const int32 ClusterIndex, const int32 PassIndex ) override;
 	virtual const int32 GetNumResults( const int32 ClusterIndex ) override;
 	virtual const int32 GetCurrentPassIndex( const int32 ClusterIndex ) override;
diff --git a/Engine/Source/Developer/AutomationController/Public/Interfaces/IAutomationReport.h b/Engine/Source/Developer/AutomationController/Public/Interfaces/IAutomationReport.h
--- a/Engine/Source/Developer/AutomationController/Public/Interfaces/IAutomationReport.h
+++ b/Engine/Source/Developer/AutomationController/Public/Interfaces/IAutomationReport.h
@@ -4,6 +4,8 @@
 
 #include "CoreMinimal.h"
 #include "Misc/AutomationTest.h"
+#include "UObject/ObjectMacros.h"
+#include "IAutomationReport.generated.h"
 
 class IAutomationReport;
 template< typename ItemType > class TFilterCollection;
@@ -20,50 +22,63 @@ typedef TSharedRef<class IAutomationReport> IAutomationReportRef;
 
 
 /** Enumeration of unit test status for special dialog */
-namespace EAutomationState
+UENUM()
+enum class EAutomationState : uint8
 {
-	enum Type
-	{
-		NotRun,					// Automation test was not run
-		InProcess,				// Automation test is running now
-		Fail,					// Automation test was run and failed
-		Success,				// Automation test was run and succeeded
-		NotEnoughParticipants,	// Automation test was not run due to number of participants
-	};
-
-	inline const TCHAR* ToString(EAutomationState::Type InType)
+	NotRun,					// Automation test was not run
+	InProcess,				// Automation test is running now
+	Fail,					// Automation test was run and failed
+	Success,				// Automation test was run and succeeded
+	NotEnoughParticipants,	// Automation test was not run due to number of participants
+};
+
+
+inline const TCHAR* ToString(EAutomationState InType)
+{
+	switch (InType)
 	{
-		switch (InType)
+		case (EAutomationState::NotRun) :
+		{
+			return TEXT("NotRun");
+		}
+		case (EAutomationState::Fail) :
 		{
-			case (EAutomationState::NotRun) :
-			{
-				return TEXT("NotRun");
-			}
-			case (EAutomationState::Fail) :
-			{
-				return TEXT("Fail");
-			}
-			case (EAutomationState::Success) :
-			{
-				return TEXT("Pass");
-			}
-			case (EAutomationState::InProcess) :
-			{
-				return TEXT("InProgress");
-			}
-			case (EAutomationState::NotEnoughParticipants) :
-			{
-				return TEXT("NotEnoughParticipants");
-			}
-			default:
-			{
-				return TEXT("Invalid");
-			}
+			return TEXT("Fail");
+		}
+		case (EAutomationState::Success) :
+		{
+			return TEXT("Pass");
+		}
+		case (EAutomationState::InProcess) :
+		{
+			return TEXT("InProgress");
+		}
+		case (EAutomationState::NotEnoughParticipants) :
+		{
+			return TEXT("NotEnoughParticipants");
+		}
+		default:
+		{
+			return TEXT("Invalid");
 		}
-		return TEXT("Invalid");
 	}
-};
+	return TEXT("Invalid");
+}
 
+struct FAutomationArtifact
+{
+public:
+	FAutomationArtifact(FString InName, FString InFilePath)
+		: Name(InName)
+		, FilePath(InFilePath)
+	{
+	}
+
+public:
+
+	FString Name;
+	FString FilePath;
+};
 
 /**
  * A struct to maintain a collection of data which was reported as part of an automation test result.
@@ -80,7 +95,7 @@ struct FAutomationTestResults
 	}
 
 	/* The current state of this test */
-	EAutomationState::Type State;
+	EAutomationState State;
 
 	/* All errors reported as part of this test */
 	TArray<FAutomationEvent> Errors;
@@ -96,6 +111,9 @@ struct FAutomationTestResults
 
 	/* The name of the instance which reported these results */
 	FString GameInstance;
+
+	/** Artifacts generated during the run of the test. */
+	TArray<FAutomationArtifact> Artifacts;
 };
 
 
@@ -311,7 +329,9 @@ public:
 	 * @param PassIndex Which test pass these results are for.
 	 * @param InResults The new set of results.
 	 */
-	virtual void SetResults( const int32 ClusterIndex, const int32 PassIndex, const FAutomationTestResults& InResults ) = 0;
+	virtual void SetResults(const int32 ClusterIndex, const int32 PassIndex, const FAutomationTestResults& InResults) = 0;
+
+	virtual void AddArtifact(const int32 ClusterIndex, const int32 PassIndex, const FAutomationArtifact& Artifact) = 0;
 
 	/**
 	 * Returns completion statistics for this branch of the testing hierarchy.
@@ -329,7 +349,7 @@ public:
 	 * @param PassIndex Which test pass to get the state of.
 	 * @return the current state of the test.
 	 */
-	virtual EAutomationState::Type GetState(const int32 ClusterIndex, const int32 PassIndex) const = 0;
+	virtual EAutomationState GetState(const int32 ClusterIndex, const int32 PassIndex) const = 0;
 
 	/**
 	 * Gets a copy of errors and warnings that were found
diff --git a/Engine/Source/Developer/AutomationWindow/Private/SAutomationGraphicalResultBox.cpp b/Engine/Source/Developer/AutomationWindow/Private/SAutomationGraphicalResultBox.cpp
--- a/Engine/Source/Developer/AutomationWindow/Private/SAutomationGraphicalResultBox.cpp
+++ b/Engine/Source/Developer/AutomationWindow/Private/SAutomationGraphicalResultBox.cpp
@@ -279,7 +279,7 @@ void SAutomationGraphicalResultBox::CreateWidgets()
 }
 
 
-FSlateColor SAutomationGraphicalResultBox::GetColorForTestState(const EAutomationState::Type TestState, const bool bHasWarnings) const
+FSlateColor SAutomationGraphicalResultBox::GetColorForTestState(const EAutomationState TestState, const bool bHasWarnings) const
 {
 	switch(TestState)
 	{
diff --git a/Engine/Source/Developer/AutomationWindow/Private/SAutomationGraphicalResultBox.h b/Engine/Source/Developer/AutomationWindow/Private/SAutomationGraphicalResultBox.h
--- a/Engine/Source/Developer/AutomationWindow/Private/SAutomationGraphicalResultBox.h
+++ b/Engine/Source/Developer/AutomationWindow/Private/SAutomationGraphicalResultBox.h
@@ -87,7 +87,7 @@ private:
 		float Duration;
 
 		/** State of the test */
-		EAutomationState::Type TestState;
+		EAutomationState TestState;
 
 		/** If the test had any warnings */
 		bool bHasWarnings;
@@ -174,7 +174,7 @@ private:
 	 * @param TestState The test state from the automation report (Success, Fail, NotRun, etc.)
 	 * @param bHasWarnings If the test had any warning logs.
 	 */
-	FSlateColor GetColorForTestState(const EAutomationState::Type TestState, const bool bHasWarnings) const;
+	FSlateColor GetColorForTestState(const EAutomationState TestState, const bool bHasWarnings) const;
 
 	/**
 	 * Gets the text to display for this test item based off the current DisplayType.
diff --git a/Engine/Source/Developer/AutomationWindow/Private/SAutomationTestItem.cpp b/Engine/Source/Developer/AutomationWindow/Private/SAutomationTestItem.cpp
--- a/Engine/Source/Developer/AutomationWindow/Private/SAutomationTestItem.cpp
+++ b/Engine/Source/Developer/AutomationWindow/Private/SAutomationTestItem.cpp
@@ -397,7 +397,7 @@ FText SAutomationTestItem::GetTestToolTip( int32 ClusterIndex ) const
 {
 	FText TestToolTip;
 	const int32 PassIndex = TestStatus->GetCurrentPassIndex(ClusterIndex);
-	EAutomationState::Type TestState = TestStatus->GetState( ClusterIndex, PassIndex );
+	EAutomationState TestState = TestStatus->GetState( ClusterIndex, PassIndex );
 	if ( TestState == EAutomationState::NotRun )
 	{
 		TestToolTip = LOCTEXT("TestToolTipNotRun", "Not Run");
@@ -439,7 +439,7 @@ FSlateColor SAutomationTestItem::ItemStatus_BackgroundColor(const int32 ClusterI
 	if (TestStatus->GetTotalNumChildren()==0)
 	{
 		const int32 PassIndex = TestStatus->GetCurrentPassIndex(ClusterIndex);
-		EAutomationState::Type TestState = TestStatus->GetState(ClusterIndex,PassIndex);
+		EAutomationState TestState = TestStatus->GetState(ClusterIndex,PassIndex);
 		if (TestState == EAutomationState::Fail)
 		{
 			// Failure is marked by a red background.
@@ -500,7 +500,7 @@ FText SAutomationTestItem::ItemStatus_DurationText() const
 EVisibility SAutomationTestItem::ItemStatus_GetStatusVisibility(const int32 ClusterIndex, const bool bForInProcessThrobber) const
 {
 	const int32 PassIndex = TestStatus->GetCurrentPassIndex(ClusterIndex);
-	EAutomationState::Type TestState = TestStatus->GetState(ClusterIndex,PassIndex);
+	EAutomationState TestState = TestStatus->GetState(ClusterIndex,PassIndex);
 	bool bImageVisible = TestState != EAutomationState::InProcess;
 
 	bool bFinalVisibility =  bForInProcessThrobber ? !bImageVisible : bImageVisible;
@@ -570,7 +570,7 @@ TOptional<float> SAutomationTestItem::ItemStatus_ProgressFraction(const int32 Cl
 const FSlateBrush* SAutomationTestItem::ItemStatus_StatusImage(const int32 ClusterIndex) const
 {
 	const int32 PassIndex = TestStatus->GetCurrentPassIndex(ClusterIndex);
-	EAutomationState::Type TestState = TestStatus->GetState(ClusterIndex,PassIndex);
+	EAutomationState TestState = TestStatus->GetState(ClusterIndex,PassIndex);
 
 	const FSlateBrush* ImageToUse;
 	switch( TestState )
diff --git a/Engine/Source/Developer/ScreenShotComparisonTools/Private/ScreenShotManager.cpp b/Engine/Source/Developer/ScreenShotComparisonTools/Private/ScreenShotManager.cpp
--- a/Engine/Source/Developer/ScreenShotComparisonTools/Private/ScreenShotManager.cpp
+++ b/Engine/Source/Developer/ScreenShotComparisonTools/Private/ScreenShotManager.cpp
@@ -54,6 +54,11 @@ FString FScreenShotManager::GetLocalApprovedFolder() const
 	return FPaths::ConvertRelativePathToFull(ScreenshotApprovedFolder);
 }
 
+FString FScreenShotManager::GetLocalComparisonFolder() const
+{
+	return FPaths::ConvertRelativePathToFull(ScreenshotDeltaFolder);
+}
+
 /* IScreenShotManager event handlers
  *****************************************************************************/
 
diff --git a/Engine/Source/Developer/ScreenShotComparisonTools/Private/ScreenShotManager.h b/Engine/Source/Developer/ScreenShotComparisonTools/Private/ScreenShotManager.h
--- a/Engine/Source/Developer/ScreenShotComparisonTools/Private/ScreenShotManager.h
+++ b/Engine/Source/Developer/ScreenShotComparisonTools/Private/ScreenShotManager.h
@@ -43,6 +43,8 @@ public:
 
 	virtual FString GetLocalApprovedFolder() const override;
 
+	virtual FString GetLocalComparisonFolder() const override;
+
 	//~ End IScreenShotManager Interface
 
 private:
diff --git a/Engine/Source/Developer/ScreenShotComparisonTools/Public/Interfaces/IScreenShotManager.h b/Engine/Source/Developer/ScreenShotComparisonTools/Public/Interfaces/IScreenShotManager.h
--- a/Engine/Source/Developer/ScreenShotComparisonTools/Public/Interfaces/IScreenShotManager.h
+++ b/Engine/Source/Developer/ScreenShotComparisonTools/Public/Interfaces/IScreenShotManager.h
@@ -54,4 +54,9 @@ public:
 	 * 
 	 */
 	virtual FString GetLocalApprovedFolder() const = 0;
+
+	/**
+	 * 
+	 */
+	virtual FString GetLocalComparisonFolder() const = 0;
 };
diff --git a/Engine/Source/Runtime/Core/Private/Misc/AutomationTest.cpp b/Engine/Source/Runtime/Core/Private/Misc/AutomationTest.cpp
--- a/Engine/Source/Runtime/Core/Private/Misc/AutomationTest.cpp
+++ b/Engine/Source/Runtime/Core/Private/Misc/AutomationTest.cpp
@@ -757,11 +757,6 @@ void FAutomationTestBase::AddAnalyticsItem(const FString& InAnalyticsItem)
 	ExecutionInfo.AnalyticsItems.Add(InAnalyticsItem);
 }
 
-void FAutomationTestBase::AddRemoteArtifact(const FString& ArtifactName, const FString& RemoteArtifactPath)
-{
-	ExecutionInfo.Artifacts.Add(FAutomationTestArtifact(ArtifactName, RemoteArtifactPath));
-}
-
 bool FAutomationTestBase::HasAnyErrors() const
 {
 	return ExecutionInfo.Errors.Num() > 0;
diff --git a/Engine/Source/Runtime/Core/Public/Misc/AutomationTest.h b/Engine/Source/Runtime/Core/Public/Misc/AutomationTest.h
--- a/Engine/Source/Runtime/Core/Public/Misc/AutomationTest.h
+++ b/Engine/Source/Runtime/Core/Public/Misc/AutomationTest.h
@@ -92,19 +92,6 @@ struct CORE_API FAutomationEvent
 	FString ToString() const;
 };
 
-class FAutomationTestArtifact
-{
-public:
-	FString Name;
-	FString FilePath;
-
-	FAutomationTestArtifact(FString InName, FString InFilePath)
-		: Name(InName)
-		, FilePath(InFilePath)
-	{
-	}
-};
-
 /** Simple class to store the results of the execution of a automation test */
 class FAutomationTestExecutionInfo
 {
@@ -130,7 +117,6 @@ public:
 		Warnings.Empty();
 		LogItems.Empty();
 		AnalyticsItems.Empty();
-		Artifacts.Empty();
 	}
 
 	/**
@@ -154,9 +140,6 @@ public:
 	/** Any analytics items that occurred during execution */
 	TArray<FString> AnalyticsItems;
 
-	/** Any items that need to be reference by reports and other systems. */
-	TArray<FAutomationTestArtifact> Artifacts;
-
 	/** Time to complete the task */
 	float Duration;
 };
@@ -974,11 +957,6 @@ public:
 	 */
 	virtual void AddAnalyticsItem(const FString& InAnalyticsItem);
 
-	/**
-	 * Occasionally a test may need to provide artifacts that should appear in the build report for the automation test.
-	 */
-	virtual void AddRemoteArtifact(const FString& ArtifactName, const FString& RemoteArtifactPath);
-
 	/**
 	 * Returns whether this test has any errors associated with it or not
 	 *
