Automation - Some additional tweaks to where reports are generated.

diff --git a/Engine/Source/Developer/AutomationController/Private/AutomationControllerManger.cpp b/Engine/Source/Developer/AutomationController/Private/AutomationControllerManger.cpp
--- a/Engine/Source/Developer/AutomationController/Private/AutomationControllerManger.cpp
+++ b/Engine/Source/Developer/AutomationController/Private/AutomationControllerManger.cpp
@@ -426,7 +426,7 @@ void FAutomationControllerManager::GenerateHtmlTestPassSummary(FDateTime Timesta
 			Logs += FString::Format(*LogTemplate, Args);
 		}
 
-		FString AritfactDirectory = FString::Printf(TEXT("ReportArtifact-%d"), +FEngineVersion::Current().GetChangelist());
+		FString AritfactDirectory = FString::Printf(TEXT("ReportArtifact-%d"), FEngineVersion::Current().GetChangelist());
 		for ( const FAutomationArtifact& Artifact : Test.Artifacts )
 		{
 			FString ArtifactFile = AritfactDirectory / FGuid::NewGuid().ToString(EGuidFormats::Digits) + FPaths::GetExtension(Artifact.FilePath, true);
