Automation - Continued itteration on the reporting mechanism, now with inline comparison controls for comparison artifacts.  Giving the system the ability to have a few hardcoded options for artifact types, for now standalone images and comparisons, that require 3 images.  Continued itteration on the look of the reports.

diff --git a/Engine/Source/Developer/AutomationController/Private/AutomationControllerManger.cpp b/Engine/Source/Developer/AutomationController/Private/AutomationControllerManger.cpp
--- a/Engine/Source/Developer/AutomationController/Private/AutomationControllerManger.cpp
+++ b/Engine/Source/Developer/AutomationController/Private/AutomationControllerManger.cpp
@@ -221,9 +221,12 @@ void FAutomationControllerManager::ProcessComparisonQueue()
 				FString UnapprovedFolder = ScreenshotManager->GetLocalUnapprovedFolder();
 				FString ComparisonFolder = ScreenshotManager->GetLocalComparisonFolder();
 
-				Report->AddArtifact(ClusterIndex, CurrentTestPass, FAutomationArtifact(TEXT("Approved"), ApprovedFolder / Result.ApprovedFile));
-				Report->AddArtifact(ClusterIndex, CurrentTestPass, FAutomationArtifact(TEXT("Incoming"), UnapprovedFolder / Result.IncomingFile));
-				Report->AddArtifact(ClusterIndex, CurrentTestPass, FAutomationArtifact(TEXT("Difference"), ComparisonFolder / Result.ComparisonFile));
+				TArray<FString> Files;
+				Files.Add(ApprovedFolder / Result.ApprovedFile);
+				Files.Add(UnapprovedFolder / Result.IncomingFile);
+				Files.Add(ComparisonFolder / Result.ComparisonFile);
+
+				Report->AddArtifact(ClusterIndex, CurrentTestPass, FAutomationArtifact(Entry->Name, EAutomationArtifactType::Comparison, Files));
 			}
 		}
 	}
@@ -380,16 +383,17 @@ void FAutomationControllerManager::GenerateHtmlTestPassSummary(FDateTime Timesta
 		return A.TestDisplayName < B.TestDisplayName;
 	});
 
-	FString MasterTemplate, ResultTemplate, LogTemplate, ArtifactTemplate;
+	FString MasterTemplate, ResultTemplate, LogTemplate, ArtifactCompareTemplate, ArtifactImageTemplate;
 	const bool bLoadedMaster = FFileHelper::LoadFileToString(MasterTemplate, *( FPaths::EngineContentDir() / TEXT("Automation/Report-Master-Template.html") ));
 	const bool bLoadedResult = FFileHelper::LoadFileToString(ResultTemplate, *( FPaths::EngineContentDir() / TEXT("Automation/Report-Result-Template.html") ));
 	const bool bLoadedLog = FFileHelper::LoadFileToString(LogTemplate, *( FPaths::EngineContentDir() / TEXT("Automation/Report-Log-Template.html") ));
-	const bool bLoadedArtifact = FFileHelper::LoadFileToString(ArtifactTemplate, *( FPaths::EngineContentDir() / TEXT("Automation/Report-Artifact-Template.html") ));
-	check(bLoadedMaster && bLoadedResult && bLoadedLog && bLoadedArtifact);
+	const bool bLoadedCompareArtifact = FFileHelper::LoadFileToString(ArtifactCompareTemplate, *( FPaths::EngineContentDir() / TEXT("Automation/Report-Artifact-Compare-Template.html") ));
+	const bool bLoadedImageArtifact = FFileHelper::LoadFileToString(ArtifactImageTemplate, *( FPaths::EngineContentDir() / TEXT("Automation/Report-Artifact-Image-Template.html") ));
+	check(bLoadedMaster && bLoadedResult && bLoadedLog && bLoadedCompareArtifact && bLoadedImageArtifact);
 
 	FString SuccessColor = TEXT("#3dd94a");
 	FString WarningColor = TEXT("#f9bb00");
-	FString ErrorColor = TEXT("#ff2154");
+	FString ErrorColor = TEXT("#e74c3c");
 
 	FString TitleColor = SuccessColor;
 	if ( SerializedPassResults.TestInformation.Num() > 0 )
@@ -426,17 +430,30 @@ void FAutomationControllerManager::GenerateHtmlTestPassSummary(FDateTime Timesta
 			Logs += FString::Format(*LogTemplate, Args);
 		}
 
-		FString AritfactDirectory = FString::Printf(TEXT("ReportArtifact-%d"), FEngineVersion::Current().GetChangelist());
 		for ( const FAutomationArtifact& Artifact : Test.Artifacts )
 		{
-			FString ArtifactFile = AritfactDirectory / FGuid::NewGuid().ToString(EGuidFormats::Digits) + FPaths::GetExtension(Artifact.FilePath, true);
-			FString ArtifactDestination = ReportOutputFolder / ArtifactFile;
-			IFileManager::Get().Copy(*ArtifactDestination, *Artifact.FilePath, true, true);
+			if ( Artifact.Type == EAutomationArtifactType::Comparison )
+			{
+				TMap<FString, FStringFormatArg> Args;
+				Args.Add(TEXT("Name"), FPlatformHttp::HtmlEncode(Artifact.Name));
+				Args.Add(TEXT("Approved"), CopyArtifact(Artifact.FilePaths[0]));
+				Args.Add(TEXT("Unapproved"), CopyArtifact(Artifact.FilePaths[1]));
+				Args.Add(TEXT("Difference"), CopyArtifact(Artifact.FilePaths[2]));
 
-			TMap<FString, FStringFormatArg> Args;
-			Args.Add(TEXT("Name"), FPlatformHttp::HtmlEncode(Artifact.Name));
-			Args.Add(TEXT("File"), ArtifactFile);
-			Logs += FString::Format(*ArtifactTemplate, Args);
+				Logs += FString::Format(*ArtifactCompareTemplate, Args);
+			}
+			else if ( Artifact.Type == EAutomationArtifactType::Image )
+			{
+				TMap<FString, FStringFormatArg> Args;
+				Args.Add(TEXT("Name"), FPlatformHttp::HtmlEncode(Artifact.Name));
+				Args.Add(TEXT("File"), CopyArtifact(Artifact.FilePaths[0]));
+
+				Logs += FString::Format(*ArtifactCompareTemplate, Args);
+			}
+			else
+			{
+				check(false);
+			}
 		}
 
 		{
@@ -466,6 +483,17 @@ void FAutomationControllerManager::GenerateHtmlTestPassSummary(FDateTime Timesta
 	}
 }
 
+FString FAutomationControllerManager::CopyArtifact(const FString& SourceFile)
+{
+	FString ReportOutputFolder = ReportOutputPathOverride.IsEmpty() ? FPaths::AutomationLogDir() : ReportOutputPathOverride;
+	FString AritfactDirectory = FString::Printf(TEXT("ReportArtifact-%d"), FEngineVersion::Current().GetChangelist());
+
+	FString ArtifactFile = AritfactDirectory / FGuid::NewGuid().ToString(EGuidFormats::Digits) + FPaths::GetExtension(SourceFile, true);
+	FString ArtifactDestination = ReportOutputFolder / ArtifactFile;
+	IFileManager::Get().Copy(*ArtifactDestination, *SourceFile, true, true);
+	return ArtifactFile;
+}
+
 void FAutomationControllerManager::ExecuteNextTask( int32 ClusterIndex, OUT bool& bAllTestsCompleted )
 {
 	bool bTestThatRequiresMultiplePraticipantsHadEnoughParticipants = false;
@@ -895,6 +923,7 @@ void FAutomationControllerManager::HandleReceivedScreenShot(const FAutomationWor
 
 	TSharedRef<FComparisonEntry> Comparison = MakeShareable(new FComparisonEntry());
 	Comparison->Sender = Context->GetSender();
+	Comparison->Name = Message.Metadata.Name;
 	Comparison->PendingComparison = ScreenshotManager->CompareScreensotAsync(Message.ScreenShotName);
 
 	ComparisonQueue.Enqueue(Comparison);
