Automation - Continued itteration on the reports, clearing artifacts after reports are generated from memory.

diff --git a/Engine/Source/Developer/AutomationController/Private/AutomationControllerManager.h b/Engine/Source/Developer/AutomationController/Private/AutomationControllerManager.h
--- a/Engine/Source/Developer/AutomationController/Private/AutomationControllerManager.h
+++ b/Engine/Source/Developer/AutomationController/Private/AutomationControllerManager.h
@@ -313,6 +313,8 @@ protected:
 	 */
 	void CheckChildResult( TSharedPtr< IAutomationReport > InReport );
 
+	FString CopyArtifact(const FString& SourceFile);
+
 	/**
 	 * Execute the next task thats available.
 	 *
@@ -484,6 +486,7 @@ private:
 	struct FComparisonEntry
 	{
 		FMessageAddress Sender;
+		FString Name;
 		TFuture<FImageComparisonResult> PendingComparison;
 	};
 
diff --git a/Engine/Source/Developer/AutomationController/Private/AutomationControllerManger.cpp b/Engine/Source/Developer/AutomationController/Private/AutomationControllerManger.cpp
--- a/Engine/Source/Developer/AutomationController/Private/AutomationControllerManger.cpp
+++ b/Engine/Source/Developer/AutomationController/Private/AutomationControllerManger.cpp
@@ -486,7 +486,7 @@ void FAutomationControllerManager::GenerateHtmlTestPassSummary(FDateTime Timesta
 FString FAutomationControllerManager::CopyArtifact(const FString& SourceFile)
 {
 	FString ReportOutputFolder = ReportOutputPathOverride.IsEmpty() ? FPaths::AutomationLogDir() : ReportOutputPathOverride;
-	FString AritfactDirectory = FString::Printf(TEXT("ReportArtifact-%d"), FEngineVersion::Current().GetChangelist());
+	FString AritfactDirectory = FString::Printf(TEXT("ReportArtifacts-%d"), FEngineVersion::Current().GetChangelist());
 
 	FString ArtifactFile = AritfactDirectory / FGuid::NewGuid().ToString(EGuidFormats::Digits) + FPaths::GetExtension(SourceFile, true);
 	FString ArtifactDestination = ReportOutputFolder / ArtifactFile;
diff --git a/Engine/Source/Developer/AutomationController/Private/AutomationReport.cpp b/Engine/Source/Developer/AutomationController/Private/AutomationReport.cpp
--- a/Engine/Source/Developer/AutomationController/Private/AutomationReport.cpp
+++ b/Engine/Source/Developer/AutomationController/Private/AutomationReport.cpp
@@ -374,11 +374,14 @@ void FAutomationReport::ResetForExecution(const int32 NumTestPasses)
 			{
 				//reset all stats
 				Results[ClusterIndex][PassIndex].State = EAutomationState::NotRun;
+				Results[ClusterIndex][PassIndex].Logs.Empty();
 				Results[ClusterIndex][PassIndex].Warnings.Empty();
 				Results[ClusterIndex][PassIndex].Errors.Empty();
+				Results[ClusterIndex][PassIndex].Artifacts.Empty();
 			}
 		}
 	}
+
 	//recurse to children
 	for (int32 ChildIndex = 0; ChildIndex < ChildReports.Num(); ++ChildIndex)
 	{
diff --git a/Engine/Source/Developer/AutomationController/Public/Interfaces/IAutomationReport.h b/Engine/Source/Developer/AutomationController/Public/Interfaces/IAutomationReport.h
--- a/Engine/Source/Developer/AutomationController/Public/Interfaces/IAutomationReport.h
+++ b/Engine/Source/Developer/AutomationController/Public/Interfaces/IAutomationReport.h
@@ -65,19 +65,27 @@ inline const TCHAR* ToString(EAutomationState InType)
 	return TEXT("Invalid");
 }
 
+enum class EAutomationArtifactType : uint8
+{
+	Image,
+	Comparison
+};
+
 struct FAutomationArtifact
 {
 public:
-	FAutomationArtifact(FString InName, FString InFilePath)
+	FAutomationArtifact(const FString& InName, EAutomationArtifactType InType, const TArray<FString>& InFilePaths)
 		: Name(InName)
-		, FilePath(InFilePath)
+		, Type(InType)
+		, FilePaths(InFilePaths)
 	{
 	}
 
 public:
 
 	FString Name;
-	FString FilePath;
+	EAutomationArtifactType Type;
+	TArray<FString> FilePaths;
 };
 
 /**
