Automation - Continuing to tweak the output folder system, so now it generates a structure better suited for deployment to a website.

diff --git a/Engine/Source/Developer/AutomationController/Private/AutomationControllerManager.h b/Engine/Source/Developer/AutomationController/Private/AutomationControllerManager.h
--- a/Engine/Source/Developer/AutomationController/Private/AutomationControllerManager.h
+++ b/Engine/Source/Developer/AutomationController/Private/AutomationControllerManager.h
@@ -313,7 +313,9 @@ protected:
 	 */
 	void CheckChildResult( TSharedPtr< IAutomationReport > InReport );
 
-	FString CopyArtifact(const FString& SourceFile);
+	FString CopyArtifact(const FString& DestFolder, const FString& SourceFile) const;
+
+	FString GetReportPath(FDateTime Timestamp) const;
 
 	/**
 	 * Execute the next task thats available.
diff --git a/Engine/Source/Developer/AutomationController/Private/AutomationControllerManger.cpp b/Engine/Source/Developer/AutomationController/Private/AutomationControllerManger.cpp
--- a/Engine/Source/Developer/AutomationController/Private/AutomationControllerManger.cpp
+++ b/Engine/Source/Developer/AutomationController/Private/AutomationControllerManger.cpp
@@ -327,9 +327,9 @@ void FAutomationControllerManager::GenerateJsonTestPassSummary(FDateTime Timesta
 	TSharedPtr<FJsonObject> ReportJson = FJsonObjectConverter::UStructToJsonObject(SerializedPassResults);
 	if (ReportJson.IsValid())
 	{
-		FString ReportOutputFolder = ReportOutputPathOverride.IsEmpty() ? FPaths::AutomationLogDir() : ReportOutputPathOverride;
+		FString ReportOutputPath = GetReportPath(Timestamp);
 
-		FString ReportFileName = FString::Printf(TEXT("%s/AutomationReport-%d-%s.json"), *ReportOutputFolder, FEngineVersion::Current().GetChangelist(), *Timestamp.ToString());
+		FString ReportFileName = FString::Printf(TEXT("%s/index.json"), *ReportOutputPath);
 		FArchive* ReportFileWriter = IFileManager::Get().CreateFileWriter(*ReportFileName);
 		if (ReportFileWriter != nullptr)
 		{
@@ -352,7 +352,7 @@ void FAutomationControllerManager::GenerateHtmlTestPassSummary(FDateTime Timesta
 		return;
 	}
 
-	FString ReportOutputFolder = ReportOutputPathOverride.IsEmpty() ? FPaths::AutomationLogDir() : ReportOutputPathOverride;
+	FString ReportOutputPath = GetReportPath(Timestamp);
 
 	FAutomatedTestPassResults SerializedPassResults = OurPassResults;
 	SerializedPassResults.TestInformation.StableSort([] (const FAutomatedTestResult& A, const FAutomatedTestResult& B) {
@@ -436,9 +436,9 @@ void FAutomationControllerManager::GenerateHtmlTestPassSummary(FDateTime Timesta
 			{
 				TMap<FString, FStringFormatArg> Args;
 				Args.Add(TEXT("Name"), FPlatformHttp::HtmlEncode(Artifact.Name));
-				Args.Add(TEXT("Approved"), CopyArtifact(Artifact.FilePaths[0]));
-				Args.Add(TEXT("Unapproved"), CopyArtifact(Artifact.FilePaths[1]));
-				Args.Add(TEXT("Difference"), CopyArtifact(Artifact.FilePaths[2]));
+				Args.Add(TEXT("Approved"), CopyArtifact(ReportOutputPath, Artifact.FilePaths[0]));
+				Args.Add(TEXT("Unapproved"), CopyArtifact(ReportOutputPath, Artifact.FilePaths[1]));
+				Args.Add(TEXT("Difference"), CopyArtifact(ReportOutputPath, Artifact.FilePaths[2]));
 
 				Logs += FString::Format(*ArtifactCompareTemplate, Args);
 			}
@@ -446,7 +446,7 @@ void FAutomationControllerManager::GenerateHtmlTestPassSummary(FDateTime Timesta
 			{
 				TMap<FString, FStringFormatArg> Args;
 				Args.Add(TEXT("Name"), FPlatformHttp::HtmlEncode(Artifact.Name));
-				Args.Add(TEXT("File"), CopyArtifact(Artifact.FilePaths[0]));
+				Args.Add(TEXT("File"), CopyArtifact(ReportOutputPath, Artifact.FilePaths[0]));
 
 				Logs += FString::Format(*ArtifactCompareTemplate, Args);
 			}
@@ -475,7 +475,7 @@ void FAutomationControllerManager::GenerateHtmlTestPassSummary(FDateTime Timesta
 
 		FString Html = FString::Format(*MasterTemplate, Args);
 
-		FString ReportFileName = FString::Printf(TEXT("%s/AutomationReport-%d-%s.html"), *ReportOutputFolder, FEngineVersion::Current().GetChangelist(), *Timestamp.ToString());
+		FString ReportFileName = FString::Printf(TEXT("%s/index.html"), *ReportOutputPath);
 		if ( !FFileHelper::SaveStringToFile(Html, *ReportFileName, FFileHelper::EEncodingOptions::ForceUTF8) )
 		{
 			GLog->Logf(ELogVerbosity::Error, TEXT("Test Report Html is invalid - report not generated."));
@@ -483,17 +483,22 @@ void FAutomationControllerManager::GenerateHtmlTestPassSummary(FDateTime Timesta
 	}
 }
 
-FString FAutomationControllerManager::CopyArtifact(const FString& SourceFile)
+FString FAutomationControllerManager::CopyArtifact(const FString& DestFolder, const FString& SourceFile) const
 {
-	FString ReportOutputFolder = ReportOutputPathOverride.IsEmpty() ? FPaths::AutomationLogDir() : ReportOutputPathOverride;
 	FString AritfactDirectory = FString::Printf(TEXT("ReportArtifacts-%d"), FEngineVersion::Current().GetChangelist());
 
 	FString ArtifactFile = AritfactDirectory / FGuid::NewGuid().ToString(EGuidFormats::Digits) + FPaths::GetExtension(SourceFile, true);
-	FString ArtifactDestination = ReportOutputFolder / ArtifactFile;
+	FString ArtifactDestination = DestFolder / ArtifactFile;
 	IFileManager::Get().Copy(*ArtifactDestination, *SourceFile, true, true);
+
 	return ArtifactFile;
 }
 
+FString FAutomationControllerManager::GetReportPath(FDateTime Timestamp) const
+{
+	return ReportOutputPathOverride.IsEmpty() ? FString::Printf(TEXT("%s/Report-%d-%s"), *FPaths::AutomationLogDir(), FEngineVersion::Current().GetChangelist(), *Timestamp.ToString()) : ReportOutputPathOverride;
+}
+
 void FAutomationControllerManager::ExecuteNextTask( int32 ClusterIndex, OUT bool& bAllTestsCompleted )
 {
 	bool bTestThatRequiresMultiplePraticipantsHadEnoughParticipants = false;
