Automation - Tweaking the report to move more styling information into the html files, rely more on clases to style things.  Also report states now involve new icons.

diff --git a/Engine/Source/Developer/AutomationController/Private/AutomationControllerManger.cpp b/Engine/Source/Developer/AutomationController/Private/AutomationControllerManger.cpp
--- a/Engine/Source/Developer/AutomationController/Private/AutomationControllerManger.cpp
+++ b/Engine/Source/Developer/AutomationController/Private/AutomationControllerManger.cpp
@@ -392,21 +392,27 @@ void FAutomationControllerManager::GenerateHtmlTestPassSummary(FDateTime Timesta
 	const bool bLoadedImageArtifact = FFileHelper::LoadFileToString(ArtifactImageTemplate, *( FPaths::EngineContentDir() / TEXT("Automation/Report-Artifact-Image-Template.html") ));
 	check(bLoadedMaster && bLoadedResult && bLoadedLog && bLoadedCompareArtifact && bLoadedImageArtifact);
 
-	FString SuccessColor = TEXT("#3dd94a");
-	FString WarningColor = TEXT("#f9bb00");
-	FString ErrorColor = TEXT("#e74c3c");
-
-	FString TitleColor = SuccessColor;
+	FString ReportState = TEXT("success");
+	FString ReportIcon = TEXT("heartbeat");
 	if ( SerializedPassResults.TestInformation.Num() > 0 )
 	{
 		const FAutomatedTestResult& FirstTest = SerializedPassResults.TestInformation[0];
-		TitleColor = FirstTest.Errors.Num() > 0 ? ErrorColor : FirstTest.Warnings.Num() > 0 ? WarningColor : SuccessColor;
+		if ( FirstTest.Errors.Num() > 0 )
+		{
+			ReportState = TEXT("error");
+			ReportIcon = TEXT("bomb");
+		}
+		else if ( FirstTest.Warnings.Num() > 0 )
+		{
+			ReportState = TEXT("warning");
+			ReportIcon = TEXT("exclamation-triangle");
+		}
 	}
 
 	FString HtmlResults;
 	for ( const FAutomatedTestResult& Test : SerializedPassResults.TestInformation )
 	{
-		FString TestColor = Test.Errors.Num() > 0 ? ErrorColor : Test.Warnings.Num() > 0 ? WarningColor : SuccessColor;
+		FString TestState = Test.Errors.Num() > 0 ? TEXT("error") : Test.Warnings.Num() > 0 ? TEXT("warning") : TEXT("success");
 
 		FString Logs = TEXT("");
 
@@ -459,7 +465,7 @@ void FAutomationControllerManager::GenerateHtmlTestPassSummary(FDateTime Timesta
 
 		{
 			TMap<FString, FStringFormatArg> Args;
-			Args.Add(TEXT("TestColor"), TestColor);
+			Args.Add(TEXT("TestState"), TestState);
 			Args.Add(TEXT("TestName"), Test.TestDisplayName);
 			Args.Add(TEXT("TestPath"), Test.FullTestPath);
 			Args.Add(TEXT("Logs"), Logs);
@@ -471,7 +477,8 @@ void FAutomationControllerManager::GenerateHtmlTestPassSummary(FDateTime Timesta
 	{
 		TMap<FString, FStringFormatArg> Args;
 		Args.Add(TEXT("Title"), TEXT("Automation Test Results"));
-		Args.Add(TEXT("TitleColor"), TitleColor);
+		Args.Add(TEXT("ReportState"), ReportState);
+		Args.Add(TEXT("ReportIcon"), ReportIcon);
 		Args.Add(TEXT("ComparisonExportDirectory"), ExportResults.ExportPath);
 		Args.Add(TEXT("Results"), HtmlResults);
 
