Print messages when causing the editor to force-quit after running automated tests, and ignore exit codes on errors.

diff --git a/Engine/Source/Developer/AutomationController/Private/AutomationCommandline.cpp b/Engine/Source/Developer/AutomationController/Private/AutomationCommandline.cpp
--- a/Engine/Source/Developer/AutomationController/Private/AutomationCommandline.cpp
+++ b/Engine/Source/Developer/AutomationController/Private/AutomationCommandline.cpp
@@ -375,9 +375,11 @@ public:
 				{
 					if (!GIsCriticalError)
 					{
+						UE_LOG(LogAutomationCommandLine, Display, TEXT("Setting GIsCriticalError due to test failures (will cause non-zero exit code)."));
 						GIsCriticalError = AutomationController->ReportsHaveErrors();
 					}
 
+					UE_LOG(LogAutomationCommandLine, Display, TEXT("Forcing shutdown."));
 					FPlatformMisc::RequestExit(true);
 
 					// We have finished the testing, and results are available
