FunctionalTest: Adding ClientDescriptor field to the index.json output reports (in FAutomatedTestPassResults)
#rnx

diff --git a/Engine/Source/Developer/AutomationController/Private/AutomationControllerManager.h b/Engine/Source/Developer/AutomationController/Private/AutomationControllerManager.h
--- a/Engine/Source/Developer/AutomationController/Private/AutomationControllerManager.h
+++ b/Engine/Source/Developer/AutomationController/Private/AutomationControllerManager.h
@@ -81,7 +81,8 @@ struct FAutomatedTestPassResults
 
 public:
 	FAutomatedTestPassResults()
-		: ReportCreatedOn(0)
+		: ClientDescriptor()
+		, ReportCreatedOn(0)
 		, Succeeded(0)
 		, SucceededWithWarnings(0)
 		, Failed(0)
@@ -89,8 +90,25 @@ public:
 		, TotalDuration(0)
 		, ComparisonExported(false)
 	{
+		if (FEngineVersion::Current().HasChangelist())
+		{
+			ClientDescriptor = FEngineVersion::Current().GetBranch()
+				+ TEXT(" - ")
+				+ FString::FromInt(FEngineVersion::Current().GetChangelist())
+				+ TEXT(" - ");
+		}
+
+		if (FPlatformProperties::RequiresCookedData())
+		{
+			ClientDescriptor += TEXT("Cooked ");
+		}
+
+		ClientDescriptor += FPlatformProperties::IniPlatformName();
 	}
 
+	UPROPERTY()
+	FString ClientDescriptor;
+
 	UPROPERTY()
 	FDateTime ReportCreatedOn;
 
