Back out changelist 5686999. Monolithic compile errors

diff --git a/Engine/Source/Developer/AutomationController/Private/AutomationControllerManager.h b/Engine/Source/Developer/AutomationController/Private/AutomationControllerManager.h
--- a/Engine/Source/Developer/AutomationController/Private/AutomationControllerManager.h
+++ b/Engine/Source/Developer/AutomationController/Private/AutomationControllerManager.h
@@ -81,8 +81,7 @@ struct FAutomatedTestPassResults
 
 public:
 	FAutomatedTestPassResults()
-		: ClientDescriptor()
-		, ReportCreatedOn(0)
+		: ReportCreatedOn(0)
 		, Succeeded(0)
 		, SucceededWithWarnings(0)
 		, Failed(0)
@@ -90,25 +89,8 @@ public:
 		, TotalDuration(0)
 		, ComparisonExported(false)
 	{
-		if (FEngineVersion::Current().HasChangelist())
-		{
-			ClientDescriptor = FEngineVersion::Current().GetBranch()
-				+ TEXT(" - ")
-				+ FString::FromInt(FEngineVersion::Current().GetChangelist())
-				+ TEXT(" - ");
-		}
-
-		if (FPlatformProperties::RequiresCookedData())
-		{
-			ClientDescriptor += TEXT("Cooked ");
-		}
-
-		ClientDescriptor += FPlatformProperties::IniPlatformName();
 	}
 
-	UPROPERTY()
-	FString ClientDescriptor;
-
 	UPROPERTY()
 	FDateTime ReportCreatedOn;
 
