Fix duplication of output from test results

https://udn.unrealengine.com/questions/490684/view.html

diff --git a/Engine/Source/Developer/AutomationController/Private/AutomationControllerManger.cpp b/Engine/Source/Developer/AutomationController/Private/AutomationControllerManger.cpp
--- a/Engine/Source/Developer/AutomationController/Private/AutomationControllerManger.cpp
+++ b/Engine/Source/Developer/AutomationController/Private/AutomationControllerManger.cpp
@@ -1110,6 +1110,8 @@ void FAutomationControllerManager::HandleRunTestsReplyMessage(const FAutomationW
 
 #if WITH_EDITOR
 			FMessageLog AutomationTestingLog("AutomationTestingLog");
+			// we log these messages ourselves for non-editor platforms so suppress this.
+			AutomationTestingLog.SuppressLoggingToOutputLog(true);
 			AutomationTestingLog.Open();
 #endif
 
