Added support for allowing definition and execution of "groups" of tests using the legacy automation framework.

Groups are defined in DefaultEngine.ini by specifying a groupname and then a set of filters that are used to determine group membership. Groups of tests may then be executed via "RunTests Group:GroupName".

Filters can now also be specified via "RunTests Filter:System+Filter:Editor". This will correctly match filters instead of the old behavior of doing a substring search e.g. SomeNonEditorTest is executed by "RunTests Editor"

Setup some common groups of tests in EngineTest project.

Added TestFilter option to RunAutomatedTests Gauntlet script. Defaults to Project+System

diff --git a/Engine/Source/Developer/AutomationController/Private/AutomationCommandline.cpp b/Engine/Source/Developer/AutomationController/Private/AutomationCommandline.cpp
--- a/Engine/Source/Developer/AutomationController/Private/AutomationCommandline.cpp
+++ b/Engine/Source/Developer/AutomationController/Private/AutomationCommandline.cpp
@@ -14,6 +14,7 @@
 #include "Misc/Paths.h"
 #include "Misc/FileHelper.h"
 #include "AssetRegistryModule.h"
+#include "AutomationControllerManager.h"
 
 DEFINE_LOG_CATEGORY_STATIC(LogAutomationCommandLine, Log, All);
 
@@ -119,28 +120,91 @@ public:
 		return false;
 	}
 
+
 	
 	void GenerateTestNamesFromCommandLine(const TArray<FString>& AllTestNames, TArray<FString>& OutTestNames)
 	{
 		OutTestNames.Empty();
 		
-		//Split the test names up
-		TArray<FString> Filters;
-		StringCommand.ParseIntoArray(Filters, TEXT("+"), true);
+		//Split the argument names up on +
+		TArray<FString> ArgumentNames;
+		StringCommand.ParseIntoArray(ArgumentNames, TEXT("+"), true);
+
+		TArray<FAutomatedTestFilter> Filters;
+
+		// get our settings CDO where things are stored
+		UAutomatedTestSettings* Settings = UAutomatedTestSettings::StaticClass()->GetDefaultObject<UAutomatedTestSettings>();
 
-		//trim cruft from all entries
-		for (int32 FilterIndex = 0; FilterIndex < Filters.Num(); ++FilterIndex)
+		// iterate through the arguments to build a filter list by doing the following -
+		// 1) If argument is a filter (filter:system) then make sure we only filter-in tests that start with that filter
+		// 2) If argument is a group then expand that group into multiple filters based on ini entries
+		// 3) Otherwise just substring match (default behavior in 4.22 and earlier).
+		for (int32 ArgumentIndex = 0; ArgumentIndex < ArgumentNames.Num(); ++ArgumentIndex)
 		{
-			Filters[FilterIndex] = Filters[FilterIndex].TrimStart().Replace(TEXT(" "), TEXT(""));
-		}
+			const FString GroupPrefix = TEXT("Group:");
+			const FString FilterPrefix = TEXT("Filter:");
+
+			FString ArgumentName = ArgumentNames[ArgumentIndex];
+
+			// if the argument is a filter (e.g. Filter:System) then create a filter that matches from the start
+			if (ArgumentName.StartsWith(FilterPrefix))
+			{
+				FString FilterName = ArgumentName.RightChop(FilterPrefix.Len());
+
+				if (FilterName.EndsWith(TEXT(".")) == false)
+				{
+					FilterName += TEXT(".");
+				}
 
-		for ( int32 TestIndex = 0; TestIndex < AllTestNames.Num(); ++TestIndex )
+				Filters.Add(FAutomatedTestFilter(FilterName, true));
+			}
+			else if (ArgumentName.StartsWith(GroupPrefix))
+			{
+				// if the argument is a group (e.g. Group:Rendering) then seach our groups for one that matches
+				FString GroupName = ArgumentName.RightChop(GroupPrefix.Len());
+
+				FAutomatedTestGroup* GroupEntry = Settings->Groups.FindByPredicate([&GroupName](const FAutomatedTestGroup& E) {
+					return E.Name == GroupName;
+				});
+
+				if (GroupEntry)
+				{
+					// if found add all its filters to our current list
+					if (GroupEntry->Filters.Num() > 0)
+					{
+						for (const FAutomatedTestFilter& GroupFilter : GroupEntry->Filters)
+						{
+							Filters.Add(GroupFilter);
+						}
+					}
+					else
+					{
+						UE_LOG(LogAutomationCommandLine, Warning, TEXT("Group %s contains no filters"), *GroupName);
+					}
+				}
+				else
+				{
+					UE_LOG(LogAutomationCommandLine, Warning, TEXT("No matching group named %s"), *GroupName);
+				}
+			}			
+			else
+			{
+				// old behavior of just string searching
+				ArgumentName = ArgumentName.TrimStart().Replace(TEXT(" "), TEXT(""));
+				Filters.Add(FAutomatedTestFilter(ArgumentName));
+			}
+		}
+		
+		for (int32 TestIndex = 0; TestIndex < AllTestNames.Num(); ++TestIndex)
 		{
 			FString TestNamesNoWhiteSpaces = AllTestNames[TestIndex].Replace(TEXT(" "), TEXT(""));
 
-			for ( int32 FilterIndex = 0; FilterIndex < Filters.Num(); ++FilterIndex )
+			for (const FAutomatedTestFilter& Filter : Filters)
 			{
-				if ( TestNamesNoWhiteSpaces.Contains(Filters[FilterIndex]))
+				bool IsMatch = (Filter.MatchFromStart && TestNamesNoWhiteSpaces.StartsWith(Filter.Contains))
+					|| (Filter.MatchFromStart == false && TestNamesNoWhiteSpaces.Contains((Filter.Contains)));
+
+				if (IsMatch)
 				{
 					OutTestNames.Add(AllTestNames[TestIndex]);
 					TestCount++;
@@ -148,6 +212,7 @@ public:
 				}
 			}
 		}
+		
 		// If we have the TestsRun array set up and are using the same command as before, clear out already run tests. 
 		if (TestsRun.Num() > 0)
 		{
@@ -206,11 +271,11 @@ public:
 			// This can get called a number of times before a worker is ready, so be conservative in how often we warn
 			if ((++WarningCount % 5) == 0 && TimeWaiting > 10.0)
 			{
-				UE_LOG(LogAutomationCommandLine, Log, TEXT("Can't find any workers! Searching again"));
+				UE_LOG(LogAutomationCommandLine, Warning, TEXT("Can't find any workers! Searching again"));
 			}
 			else
 			{
-				UE_LOG(LogAutomationCommandLine, Warning, TEXT("Can't find any workers! Searching again"));
+				UE_LOG(LogAutomationCommandLine, Log, TEXT("Can't find any workers! Searching again"));
 			}
 
 			AutomationTestState = EAutomationTestState::FindWorkers;
diff --git a/Engine/Source/Developer/AutomationController/Private/AutomationControllerManager.h b/Engine/Source/Developer/AutomationController/Private/AutomationControllerManager.h
--- a/Engine/Source/Developer/AutomationController/Private/AutomationControllerManager.h
+++ b/Engine/Source/Developer/AutomationController/Private/AutomationControllerManager.h
@@ -154,6 +154,65 @@ public:
 	}
 };
 
+/*
+ *	Describes a filter for a test group.
+ */
+USTRUCT()
+struct FAutomatedTestFilter
+{
+	GENERATED_USTRUCT_BODY()
+
+public:
+
+	FAutomatedTestFilter(FString InContains, bool InMatchFromStart = false) 
+		: Contains(InContains), MatchFromStart(InMatchFromStart)
+	{
+	}
+
+	FAutomatedTestFilter() : FAutomatedTestFilter(TEXT("")) {}
+
+	/** String that the test must contain */
+	UPROPERTY(Config)
+	FString Contains;
+
+	/** If true start matching from the start of the string, else anywhere */
+	UPROPERTY(Config)
+	bool MatchFromStart;
+};
+
+/*
+ *	Describes a group of tests. Each group has a name and a set of filters that determine group membership
+ */
+USTRUCT()
+struct FAutomatedTestGroup
+{
+	GENERATED_USTRUCT_BODY()
+
+public:
+
+	UPROPERTY(Config)
+	FString Name;
+	
+	UPROPERTY(Config)
+	TArray<FAutomatedTestFilter> Filters;
+};
+
+
+/*
+ *	Setttings for automated tests in a project
+ */
+UCLASS(Config = Engine)
+class UAutomatedTestSettings : public UObject
+{
+	GENERATED_BODY()
+
+public:
+
+	/** List of user-defined test groups */
+	UPROPERTY(Config)
+	TArray<FAutomatedTestGroup> Groups;
+};
+
 
 /**
  * Implements the AutomationController module.
