Cleaned up some logging in automation controller

diff --git a/Engine/Source/Developer/AutomationController/Private/AutomationCommandline.cpp b/Engine/Source/Developer/AutomationController/Private/AutomationCommandline.cpp
--- a/Engine/Source/Developer/AutomationController/Private/AutomationCommandline.cpp
+++ b/Engine/Source/Developer/AutomationController/Private/AutomationCommandline.cpp
@@ -455,13 +455,17 @@ public:
 				{
 					if (!GIsCriticalError)
 					{
-						UE_LOG(LogAutomationCommandLine, Display, TEXT("Setting GIsCriticalError due to test failures (will cause non-zero exit code)."));
-						GIsCriticalError = AutomationController->ReportsHaveErrors();
+						if (AutomationController->ReportsHaveErrors())
+						{
+							UE_LOG(LogAutomationCommandLine, Display, TEXT("Setting GIsCriticalError due to test failures (will cause non-zero exit code)."));
+							GIsCriticalError = AutomationController->ReportsHaveErrors();
+						}
 					}
 
-					UE_LOG(LogAutomationCommandLine, Display, TEXT("Forcing shutdown."));
+					UE_LOG(LogAutomationCommandLine, Log, TEXT("Forcing shutdown."));
+					// some tools parse this.
+					UE_LOG(LogAutomationCommandLine, Display, TEXT("**** TEST COMPLETE. EXIT CODE: %d ****"), GIsCriticalError ? -1 : 0);
 					FPlatformMisc::RequestExit(true);
-
 					// We have finished the testing, and results are available
 					AutomationTestState = EAutomationTestState::Complete;
 				}
diff --git a/Engine/Source/Developer/AutomationController/Private/AutomationControllerManger.cpp b/Engine/Source/Developer/AutomationController/Private/AutomationControllerManger.cpp
--- a/Engine/Source/Developer/AutomationController/Private/AutomationControllerManger.cpp
+++ b/Engine/Source/Developer/AutomationController/Private/AutomationControllerManger.cpp
@@ -27,7 +27,7 @@
 #include "Logging/MessageLog.h"
 #endif
 
-DEFINE_LOG_CATEGORY_STATIC(AutomationControllerLog, Log, All)
+DEFINE_LOG_CATEGORY_STATIC(LogAutomationController, Log, All)
 
 FAutomationControllerManager::FAutomationControllerManager()
 {
@@ -256,7 +256,7 @@ void FAutomationControllerManager::ProcessComparisonQueue()
 			}
 			else
 			{
-				UE_LOG(AutomationControllerLog, Error, TEXT("Cannot generate screenshot report for screenshot %s as report is missing"), *Result.IncomingFile);
+				UE_LOG(LogAutomationController, Error, TEXT("Cannot generate screenshot report for screenshot %s as report is missing"), *Result.IncomingFile);
 			}
 		}
 	}
@@ -303,10 +303,10 @@ void FAutomationControllerManager::ProcessAvailableTasks()
 
 void FAutomationControllerManager::ReportTestResults()
 {
-	GLog->Logf(TEXT("Test Pass Results:"));
+	UE_LOG(LogAutomationController, Log, TEXT("Test Pass Results:"));
 	for ( int32 i = 0; i < OurPassResults.Tests.Num(); i++ )
 	{
-		GLog->Logf(TEXT("%s: %s"), *OurPassResults.Tests[i].TestDisplayName, ToString(OurPassResults.Tests[i].State));
+		UE_LOG(LogAutomationController, Log, TEXT("%s: %s"), *OurPassResults.Tests[i].TestDisplayName, ToString(OurPassResults.Tests[i].State));
 	}
 }
 
@@ -351,7 +351,7 @@ void FAutomationControllerManager::CollectTestResults(TSharedPtr<IAutomationRepo
 
 bool FAutomationControllerManager::GenerateJsonTestPassSummary(const FAutomatedTestPassResults& SerializedPassResults, FDateTime Timestamp)
 {
-	UE_LOG(AutomationControllerLog, Display, TEXT("Converting results to json object..."));
+	UE_LOG(LogAutomationController, Display, TEXT("Converting results to json object..."));
 
 	FString Json;
 	if (FJsonObjectConverter::UStructToJsonObjectString(SerializedPassResults, Json))
@@ -364,18 +364,18 @@ bool FAutomationControllerManager::GenerateJsonTestPassSummary(const FAutomatedT
 		{
 			if (FFileHelper::SaveStringToFile(Json, *ReportFileName, FFileHelper::EEncodingOptions::ForceUTF8))
 			{
-				UE_LOG(AutomationControllerLog, Display, TEXT("Successfully wrote json results file!"));
+				UE_LOG(LogAutomationController, Display, TEXT("Successfully wrote json results file!"));
 				return true;
 			}
 
 			FPlatformProcess::Sleep(SleepBetweenAttempts);
 		}
 
-		UE_LOG(AutomationControllerLog, Warning, TEXT("Failed to write test report json to '%s' after 3 attempts - No report will be generated."), *ReportFileName);
+		UE_LOG(LogAutomationController, Warning, TEXT("Failed to write test report json to '%s' after 3 attempts - No report will be generated."), *ReportFileName);
 	}
 	else
 	{
-		UE_LOG(AutomationControllerLog, Error, TEXT("Failed to convert test results to json object - No report will be generated."));
+		UE_LOG(LogAutomationController, Error, TEXT("Failed to convert test results to json object - No report will be generated."));
 	}
 		
 	return false;
@@ -383,7 +383,7 @@ bool FAutomationControllerManager::GenerateJsonTestPassSummary(const FAutomatedT
 
 bool FAutomationControllerManager::GenerateHtmlTestPassSummary(const FAutomatedTestPassResults& SerializedPassResults, FDateTime Timestamp)
 {
-	UE_LOG(AutomationControllerLog, Display, TEXT("Loading results html template..."));
+	UE_LOG(LogAutomationController, Display, TEXT("Loading results html template..."));
 
 	FString ReportTemplate;
 	if (FFileHelper::LoadFileToString(ReportTemplate, *(FPaths::EngineContentDir() / TEXT("Automation/Report-Template.html"))))
@@ -396,18 +396,18 @@ bool FAutomationControllerManager::GenerateHtmlTestPassSummary(const FAutomatedT
 		{
 			if (FFileHelper::SaveStringToFile(ReportTemplate, *ReportFileName, FFileHelper::EEncodingOptions::ForceUTF8))
 			{
-				UE_LOG(AutomationControllerLog, Display, TEXT("Successfully wrote html results file!"));
+				UE_LOG(LogAutomationController, Display, TEXT("Successfully wrote html results file!"));
 				return true;
 			}
 
 			FPlatformProcess::Sleep(SleepBetweenAttempts);
 		}
 
-		UE_LOG(AutomationControllerLog, Warning, TEXT("Failed to write test report html to '%s' after 3 attempts - No report will be generated."), *ReportFileName);
+		UE_LOG(LogAutomationController, Warning, TEXT("Failed to write test report html to '%s' after 3 attempts - No report will be generated."), *ReportFileName);
 	}
 	else
 	{
-		UE_LOG(AutomationControllerLog, Error, TEXT("Failed to load test report html template - No report will be generated."));
+		UE_LOG(LogAutomationController, Error, TEXT("Failed to load test report html template - No report will be generated."));
 	}
 	
 	return false;
@@ -485,7 +485,7 @@ void FAutomationControllerManager::ExecuteNextTask( int32 ClusterIndex, OUT bool
 						{
 							FAutomationTestResults TestResults;
 
-							GLog->Logf(ELogVerbosity::Display, TEXT("Running Automation: '%s' (Class Name: '%s')"), *TestsRunThisPass[AddressIndex]->GetFullTestPath(), *TestsRunThisPass[AddressIndex]->GetCommand());
+							UE_LOG(LogAutomationController, Display, TEXT("Running Automation: '%s' (Class Name: '%s')"), *TestsRunThisPass[AddressIndex]->GetFullTestPath(), *TestsRunThisPass[AddressIndex]->GetCommand());
 							TestResults.State = EAutomationState::InProcess;
 
 							if (CheckpointFile)
@@ -654,11 +654,11 @@ void FAutomationControllerManager::ProcessResults()
 	{
 		FDateTime Timestamp = FDateTime::Now();
 
-		UE_LOG(AutomationControllerLog, Display, TEXT("Generating Automation Report @ %s."), *ReportOutputPath);
+		UE_LOG(LogAutomationController, Display, TEXT("Generating Automation Report @ %s."), *ReportOutputPath);
 
 		if ( IFileManager::Get().DirectoryExists(*ReportOutputPath) )
 		{
-			UE_LOG(AutomationControllerLog, Display, TEXT("Existing report directory found, deleting %s."), *ReportOutputPath);
+			UE_LOG(LogAutomationController, Display, TEXT("Existing report directory found, deleting %s."), *ReportOutputPath);
 
 			// Clear the old report folder.  Why move it first?  Because RemoveDirectory
 			// is actually an async call that is not immediately carried out by the Windows OS; Moving a directory on the other hand, is sync.
@@ -725,7 +725,7 @@ void FAutomationControllerManager::ProcessResults()
 			}
 		}
 
-		UE_LOG(AutomationControllerLog, Display, TEXT("Writing reports... %s."), *ReportOutputPath);
+		UE_LOG(LogAutomationController, Display, TEXT("Writing reports... %s."), *ReportOutputPath);
 
 		// Generate Json
 		GenerateJsonTestPassSummary(SerializedPassResults, Timestamp);
@@ -735,12 +735,12 @@ void FAutomationControllerManager::ProcessResults()
 
 		if ( !DeveloperReportUrl.IsEmpty() )
 		{
-			UE_LOG(AutomationControllerLog, Display, TEXT("Launching Report URL %s."), *DeveloperReportUrl);
+			UE_LOG(LogAutomationController, Display, TEXT("Launching Report URL %s."), *DeveloperReportUrl);
 
 			FPlatformProcess::LaunchURL(*DeveloperReportUrl, nullptr, nullptr);
 		}
 
-		UE_LOG(AutomationControllerLog, Display, TEXT("Done writing reports... %s."), *ReportOutputPath);
+		UE_LOG(LogAutomationController, Display, TEXT("Done writing reports... %s."), *ReportOutputPath);
 	}
 
 	// Then clean our array for the next pass.
@@ -845,7 +845,7 @@ void FAutomationControllerManager::UpdateTests()
 				TArray<FString> ErrorStringArray;
 				ErrorStringArray.Add(FString(TEXT("Failed")));
 				bHasErrors = true;
-				GLog->Logf(ELogVerbosity::Display, TEXT("Timeout hit. Nooooooo."));
+				UE_LOG(LogAutomationController, Display, TEXT("Timeout hit. Nooooooo."));
 
 				FAutomationTestResults TestResults;
 				TestResults.State = EAutomationState::Fail;
@@ -873,13 +873,13 @@ void FAutomationControllerManager::UpdateTests()
 					// Process results first to write out the report
 					ProcessResults();
 
-					GLog->Logf(ELogVerbosity::Display, TEXT("Module disabled"));
+					UE_LOG(LogAutomationController, Display, TEXT("Module disabled"));
 					SetControllerStatus(EAutomationControllerModuleState::Disabled);
 					ClusterDistributionMask = 0;
 				}
 				else
 				{
-					GLog->Logf(ELogVerbosity::Display, TEXT("Module not disabled. Keep looking."));
+					UE_LOG(LogAutomationController, Display, TEXT("Module not disabled. Keep looking."));
 					// Remove the cluster from the mask if there are no active devices left
 					if ( DeviceClusterManager.GetNumActiveDevicesInCluster(ClusterIndex) == 0 )
 					{
@@ -1108,8 +1108,9 @@ void FAutomationControllerManager::HandleRunTestsReplyMessage(const FAutomationW
 			// Gather all of the data relevant to this test for our json reporting.
 			CollectTestResults(Report, FinalResults);
 
+			FName CategoryName = "LogAutomationController";
 #if WITH_EDITOR
-			FMessageLog AutomationTestingLog("AutomationTestingLog");
+			FMessageLog AutomationTestingLog(CategoryName);
 			// we log these messages ourselves for non-editor platforms so suppress this.
 			AutomationTestingLog.SuppressLoggingToOutputLog(true);
 			AutomationTestingLog.Open();
@@ -1120,19 +1121,19 @@ void FAutomationControllerManager::HandleRunTestsReplyMessage(const FAutomationW
 				switch (Entry.Event.Type)
 				{
 				case EAutomationEventType::Info:
-					GLog->Logf(ELogVerbosity::Log, TEXT("%s"), *Entry.ToString());
+					UE_LOG(LogAutomationController, Log, TEXT("%s"), *Entry.ToString());
 #if WITH_EDITOR
 					AutomationTestingLog.Info(FText::FromString(Entry.ToString()));
 #endif
 					break;
 				case EAutomationEventType::Warning:
-					GLog->Logf(ELogVerbosity::Warning, TEXT("%s"), *Entry.ToString());
+					UE_LOG(LogAutomationController, Warning, TEXT("%s"), *Entry.ToString());
 #if WITH_EDITOR
 					AutomationTestingLog.Warning(FText::FromString(Entry.ToString()));
 #endif
 					break;
 				case EAutomationEventType::Error:
-					GLog->Logf(ELogVerbosity::Error, TEXT("%s"), *Entry.ToString());
+					UE_LOG(LogAutomationController, Error, TEXT("%s"), *Entry.ToString());
 #if WITH_EDITOR
 					AutomationTestingLog.Error(FText::FromString(Entry.ToString()));
 #endif
@@ -1142,16 +1143,16 @@ void FAutomationControllerManager::HandleRunTestsReplyMessage(const FAutomationW
 
 			if (TestResults.State == EAutomationState::Success)
 			{
-				FString SuccessString = FString::Printf(TEXT("...Automation Test Succeeded (%s)"), *Report->GetDisplayName());
-				GLog->Logf(ELogVerbosity::Log, TEXT("%s"), *SuccessString);
+				FString SuccessString = FString::Printf(TEXT("Automation Test Succeeded (%s - %s)"), *Report->GetDisplayName(), *Report->GetFullTestPath());
+				UE_LOG(LogAutomationController, Log, TEXT("%s"), *SuccessString);
 #if WITH_EDITOR
 				AutomationTestingLog.Info(FText::FromString(*SuccessString));
 #endif
 			}
 			else
 			{
-				FString FailureString = FString::Printf(TEXT("...Automation Test Failed (%s)"), *Report->GetDisplayName());
-				GLog->Logf(ELogVerbosity::Log, TEXT("%s"), *FailureString);
+				FString FailureString = FString::Printf(TEXT("Automation Test Failed (%s - %s)"), *Report->GetDisplayName(), *Report->GetFullTestPath());
+				UE_LOG(LogAutomationController, Log, TEXT("%s"), *FailureString);
 #if WITH_EDITOR
 				AutomationTestingLog.Error(FText::FromString(*FailureString));
 #endif
@@ -1204,7 +1205,7 @@ TArray<FString> FAutomationControllerManager::GetCheckpointFileContents()
 		FileData.ParseIntoArrayLines(TestsRun);
 		for (int i = 0; i < TestsRun.Num(); i++)
 		{
-			GLog->Log(TEXT("AutomationCheckpoint"), ELogVerbosity::Log, *TestsRun[i]);
+			UE_LOG(LogAutomationController, Log, TEXT("AutomationCheckpoint %s"), *TestsRun[i]);
 		}
 	}
 	return TestsRun;
diff --git a/Engine/Source/Runtime/Core/Private/Misc/AutomationTest.cpp b/Engine/Source/Runtime/Core/Private/Misc/AutomationTest.cpp
--- a/Engine/Source/Runtime/Core/Private/Misc/AutomationTest.cpp
+++ b/Engine/Source/Runtime/Core/Private/Misc/AutomationTest.cpp
@@ -724,23 +724,24 @@ FString FAutomationExecutionEntry::ToString() const
 {
 	FString ComplexString;
 
+	ComplexString = Event.Message;
+	
 	if ( !Filename.IsEmpty() && LineNumber > 0 )
 	{
+		ComplexString += TEXT(" [");
 		ComplexString += Filename;
 		ComplexString += TEXT("(");
 		ComplexString += FString::FromInt(LineNumber);
-		ComplexString += TEXT("): ");
+		ComplexString += TEXT(")]");
 	}
 
 	if ( !Event.Context.IsEmpty() )
 	{
-		ComplexString += TEXT("[");
+		ComplexString += TEXT(" [");
 		ComplexString += Event.Context;
 		ComplexString += TEXT("] ");
 	}
 
-	ComplexString += Event.Message;
-
 	return ComplexString;
 }
 
