Moved automation settings to a common class.

Removed SetTreatWarningsAsErrors from AutomationFramework. Tests should now override TreatLogErrorsAsErrors and TreatLogWarningsAsErrors to control this behavior.

Tidied up some logging and formalized messages emitted by tests.

Emit a Gauntlet-parsed message on test completion

diff --git a/Engine/Source/Developer/AutomationController/AutomationController.Build.cs b/Engine/Source/Developer/AutomationController/AutomationController.Build.cs
--- a/Engine/Source/Developer/AutomationController/AutomationController.Build.cs
+++ b/Engine/Source/Developer/AutomationController/AutomationController.Build.cs
@@ -11,6 +11,7 @@ namespace UnrealBuildTool.Rules
 				{
 					"Core",
 					"CoreUObject",
+					"Engine"
 				}
 			); 
 			
diff --git a/Engine/Source/Developer/AutomationController/Private/AutomationCommandline.cpp b/Engine/Source/Developer/AutomationController/Private/AutomationCommandline.cpp
--- a/Engine/Source/Developer/AutomationController/Private/AutomationCommandline.cpp
+++ b/Engine/Source/Developer/AutomationController/Private/AutomationCommandline.cpp
@@ -14,7 +14,7 @@
 #include "Misc/Paths.h"
 #include "Misc/FileHelper.h"
 #include "AssetRegistryModule.h"
-#include "AutomationControllerManager.h"
+#include "Tests/AutomationTestSettings.h"
 
 DEFINE_LOG_CATEGORY_STATIC(LogAutomationCommandLine, Log, All);
 
@@ -133,7 +133,7 @@ public:
 		TArray<FAutomatedTestFilter> Filters;
 
 		// get our settings CDO where things are stored
-		UAutomatedTestSettings* Settings = UAutomatedTestSettings::StaticClass()->GetDefaultObject<UAutomatedTestSettings>();
+		UAutomationTestSettings* Settings = UAutomationTestSettings::StaticClass()->GetDefaultObject<UAutomationTestSettings>();
 
 		// iterate through the arguments to build a filter list by doing the following -
 		// 1) If argument is a filter (filter:system) then make sure we only filter-in tests that start with that filter
@@ -464,7 +464,7 @@ public:
 
 					UE_LOG(LogAutomationCommandLine, Log, TEXT("Forcing shutdown."));
 					// some tools parse this.
-					//UE_LOG(LogAutomationCommandLine, Display, TEXT("**** TEST COMPLETE. EXIT CODE: %d ****"), GIsCriticalError ? -1 : 0);
+					UE_LOG(LogAutomationCommandLine, Display, TEXT("**** TEST COMPLETE. EXIT CODE: %d ****"), GIsCriticalError ? -1 : 0);
 					FPlatformMisc::RequestExit(true);
 					// We have finished the testing, and results are available
 					AutomationTestState = EAutomationTestState::Complete;
diff --git a/Engine/Source/Developer/AutomationController/Private/AutomationControllerManager.h b/Engine/Source/Developer/AutomationController/Private/AutomationControllerManager.h
--- a/Engine/Source/Developer/AutomationController/Private/AutomationControllerManager.h
+++ b/Engine/Source/Developer/AutomationController/Private/AutomationControllerManager.h
@@ -154,64 +154,6 @@ public:
 	}
 };
 
-/*
- *	Describes a filter for a test group.
- */
-USTRUCT()
-struct FAutomatedTestFilter
-{
-	GENERATED_USTRUCT_BODY()
-
-public:
-
-	FAutomatedTestFilter(FString InContains, bool InMatchFromStart = false) 
-		: Contains(InContains), MatchFromStart(InMatchFromStart)
-	{
-	}
-
-	FAutomatedTestFilter() : FAutomatedTestFilter(TEXT("")) {}
-
-	/** String that the test must contain */
-	UPROPERTY(Config)
-	FString Contains;
-
-	/** If true start matching from the start of the string, else anywhere */
-	UPROPERTY(Config)
-	bool MatchFromStart;
-};
-
-/*
- *	Describes a group of tests. Each group has a name and a set of filters that determine group membership
- */
-USTRUCT()
-struct FAutomatedTestGroup
-{
-	GENERATED_USTRUCT_BODY()
-
-public:
-
-	UPROPERTY(Config)
-	FString Name;
-	
-	UPROPERTY(Config)
-	TArray<FAutomatedTestFilter> Filters;
-};
-
-
-/*
- *	Setttings for automated tests in a project
- */
-UCLASS(Config = Engine)
-class UAutomatedTestSettings : public UObject
-{
-	GENERATED_BODY()
-
-public:
-
-	/** List of user-defined test groups */
-	UPROPERTY(Config)
-	TArray<FAutomatedTestGroup> Groups;
-};
 
 
 /**
@@ -508,6 +450,8 @@ private:
 	/** Handles FAutomationWorkerWorkerOffline messages. */
 	void HandleWorkerOfflineMessage( const FAutomationWorkerWorkerOffline& Message, const TSharedRef<IMessageContext, ESPMode::ThreadSafe>& Context );
 
+	/** Writes out this automation result to the log */
+	void ReportAutomationResult(const TSharedPtr<IAutomationReport> InReport, int32 ClusterIndex, int32 PassIndex);
 private:
 
 	/** Session this controller is currently communicating with */
diff --git a/Engine/Source/Developer/AutomationController/Private/AutomationControllerManger.cpp b/Engine/Source/Developer/AutomationController/Private/AutomationControllerManger.cpp
--- a/Engine/Source/Developer/AutomationController/Private/AutomationControllerManger.cpp
+++ b/Engine/Source/Developer/AutomationController/Private/AutomationControllerManger.cpp
@@ -1078,6 +1078,79 @@ void FAutomationControllerManager::HandleRequestTestsReplyCompleteMessage(const
 	SetTestNames(Context->GetSender(), TestInfo);
 }
 
+
+void FAutomationControllerManager::ReportAutomationResult(const TSharedPtr<IAutomationReport> InReport, int32 ClusterIndex, int32 PassIndex)
+{
+	// these strings are parsed by Gauntlet (AutomationLogParser) so make sure changes are replicated there!
+	#define AutomationSuccessFormat		TEXT("Automation Test Succeeded (%s - %s)")
+	#define AutomationFailureFormat		TEXT("Automation Test Failed (%s - %s)")
+	#define BeginEventsFormat			TEXT("BeginEvents: %s")
+	#define EndEventsFormat				TEXT("EndEvents: %s")
+
+	FName CategoryName = "LogAutomationController";
+#if WITH_EDITOR
+	FMessageLog AutomationTestingLog(CategoryName);
+	// we log these messages ourselves for non-editor platforms so suppress this.
+	AutomationTestingLog.SuppressLoggingToOutputLog(true);
+	AutomationTestingLog.Open();
+#endif
+
+	const FAutomationTestResults& Results = InReport->GetResults(ClusterIndex, PassIndex);
+
+	if (Results.State == EAutomationState::Success)
+	{
+		FString SuccessString = FString::Printf(AutomationSuccessFormat, *InReport->GetDisplayName(), *InReport->GetFullTestPath());
+		UE_LOG(LogAutomationController, Display, TEXT("%s"), *SuccessString);
+#if WITH_EDITOR
+		AutomationTestingLog.Info(FText::FromString(*SuccessString));
+#endif
+	}
+	else
+	{
+		FString FailureString = FString::Printf(AutomationFailureFormat, *InReport->GetDisplayName(), *InReport->GetFullTestPath());
+		UE_LOG(LogAutomationController, Error, TEXT("%s"), *FailureString);
+#if WITH_EDITOR
+		AutomationTestingLog.Error(FText::FromString(*FailureString));
+#endif
+	}
+
+	// bracket these for easy parsing
+	UE_LOG(LogAutomationController, Log, BeginEventsFormat, *InReport->GetFullTestPath());
+
+	for (const FAutomationExecutionEntry& Entry : Results.GetEntries())
+	{
+		switch (Entry.Event.Type)
+		{
+		case EAutomationEventType::Info:
+			UE_LOG(LogAutomationController, Log, TEXT("%s"), *Entry.ToString());
+#if WITH_EDITOR
+			AutomationTestingLog.Info(FText::FromString(Entry.ToString()));
+#endif
+			break;
+		case EAutomationEventType::Warning:
+			UE_LOG(LogAutomationController, Warning, TEXT("%s"), *Entry.ToString());
+#if WITH_EDITOR
+			AutomationTestingLog.Warning(FText::FromString(Entry.ToString()));
+#endif
+			break;
+		case EAutomationEventType::Error:
+			UE_LOG(LogAutomationController, Error, TEXT("%s"), *Entry.ToString());
+#if WITH_EDITOR
+			AutomationTestingLog.Error(FText::FromString(Entry.ToString()));
+#endif
+			break;
+		}
+	}
+
+	UE_LOG(LogAutomationController, Log, EndEventsFormat, *InReport->GetFullTestPath());
+
+	#undef AutomationSuccessFormat
+	#undef AutomationFailureFormat
+	#undef BeginEventsFormat
+	#undef EndEventsFormat
+}
+
+
 void FAutomationControllerManager::HandleRunTestsReplyMessage(const FAutomationWorkerRunTestsReply& Message, const TSharedRef<IMessageContext, ESPMode::ThreadSafe>& Context)
 {
 	// If we should commit these results
@@ -1108,55 +1181,7 @@ void FAutomationControllerManager::HandleRunTestsReplyMessage(const FAutomationW
 			// Gather all of the data relevant to this test for our json reporting.
 			CollectTestResults(Report, FinalResults);
 
-			FName CategoryName = "LogAutomationController";
-#if WITH_EDITOR
-			FMessageLog AutomationTestingLog(CategoryName);
-			// we log these messages ourselves for non-editor platforms so suppress this.
-			AutomationTestingLog.SuppressLoggingToOutputLog(true);
-			AutomationTestingLog.Open();
-#endif
-
-			for (const FAutomationExecutionEntry& Entry : TestResults.GetEntries())
-			{
-				switch (Entry.Event.Type)
-				{
-				case EAutomationEventType::Info:
-					UE_LOG(LogAutomationController, Log, TEXT("%s"), *Entry.ToString());
-#if WITH_EDITOR
-					AutomationTestingLog.Info(FText::FromString(Entry.ToString()));
-#endif
-					break;
-				case EAutomationEventType::Warning:
-					UE_LOG(LogAutomationController, Warning, TEXT("%s"), *Entry.ToString());
-#if WITH_EDITOR
-					AutomationTestingLog.Warning(FText::FromString(Entry.ToString()));
-#endif
-					break;
-				case EAutomationEventType::Error:
-					UE_LOG(LogAutomationController, Error, TEXT("%s"), *Entry.ToString());
-#if WITH_EDITOR
-					AutomationTestingLog.Error(FText::FromString(Entry.ToString()));
-#endif
-					break;
-				}
-			}
-
-			if (TestResults.State == EAutomationState::Success)
-			{
-				FString SuccessString = FString::Printf(TEXT("Automation Test Succeeded (%s - %s)"), *Report->GetDisplayName(), *Report->GetFullTestPath());
-				UE_LOG(LogAutomationController, Log, TEXT("%s"), *SuccessString);
-#if WITH_EDITOR
-				AutomationTestingLog.Info(FText::FromString(*SuccessString));
-#endif
-			}
-			else
-			{
-				FString FailureString = FString::Printf(TEXT("Automation Test Failed (%s - %s)"), *Report->GetDisplayName(), *Report->GetFullTestPath());
-				UE_LOG(LogAutomationController, Log, TEXT("%s"), *FailureString);
-#if WITH_EDITOR
-				AutomationTestingLog.Error(FText::FromString(*FailureString));
-#endif
-			}
+			ReportAutomationResult(Report, ClusterIndex, CurrentTestPass);
 		}
 
 		// Device is now good to go
diff --git a/Engine/Source/Runtime/Core/Private/Misc/App.cpp b/Engine/Source/Runtime/Core/Private/Misc/App.cpp
--- a/Engine/Source/Runtime/Core/Private/Misc/App.cpp
+++ b/Engine/Source/Runtime/Core/Private/Misc/App.cpp
@@ -342,7 +342,7 @@ void FApp::PrintStartupLogMessages()
 	UE_LOG(LogInit, Log, TEXT("Build Configuration: %s"), EBuildConfigurations::ToString(FApp::GetBuildConfiguration()));
 	UE_LOG(LogInit, Log, TEXT("Branch Name: %s"), *FApp::GetBranchName());
 	FString FilteredString = FCommandLine::IsCommandLineLoggingFiltered() ? TEXT("Filtered ") : TEXT("");
-	UE_LOG(LogInit, Log, TEXT("%sCommand Line: %s"), *FilteredString, FCommandLine::GetForLogging());
+	UE_LOG(LogInit, Display, TEXT("%sCommand Line: %s"), *FilteredString, FCommandLine::GetForLogging());
 	UE_LOG(LogInit, Log, TEXT("Base Directory: %s"), FPlatformProcess::BaseDir());
 	//UE_LOG(LogInit, Log, TEXT("Character set: %s"), sizeof(TCHAR)==1 ? TEXT("ANSI") : TEXT("Unicode") );
 	UE_LOG(LogInit, Log, TEXT("Installed Engine Build: %d"), FApp::IsEngineInstalled() ? 1 : 0);
diff --git a/Engine/Source/Runtime/Core/Private/Misc/AutomationTest.cpp b/Engine/Source/Runtime/Core/Private/Misc/AutomationTest.cpp
--- a/Engine/Source/Runtime/Core/Private/Misc/AutomationTest.cpp
+++ b/Engine/Source/Runtime/Core/Private/Misc/AutomationTest.cpp
@@ -29,28 +29,29 @@ void FAutomationTestFramework::FAutomationTestFeedbackContext::Serialize( const
 	// Ensure there's a valid unit test associated with the context
 	if ( CurTest )
 	{
-		// Warnings
-		if ( Verbosity == ELogVerbosity::Warning )
+		bool CaptureLog = !CurTest->SuppressLogs() 
+					&& (Verbosity == ELogVerbosity::Error || Verbosity == ELogVerbosity::Warning || Verbosity == ELogVerbosity::Display);
+
+		if (CaptureLog)
 		{
-			// If warnings should be treated as errors, log the warnings as such in the current unit test
-			if ( TreatWarningsAsErrors )
+			bool IsError = (Verbosity == ELogVerbosity::Error && CurTest->TreatLogErrorsAsErrors()) || (Verbosity == ELogVerbosity::Warning && CurTest->TreatLogWarningsAsErrors());
+			bool IsWarning = (Verbosity == ELogVerbosity::Warning || Verbosity == ELogVerbosity::Error) && !IsError;
+			
+			// Errors
+			if (IsError)
 			{
 				CurTest->AddError(FString(V), STACK_OFFSET);
 			}
-			else
+			// Warnings
+			else if (IsWarning)
 			{
 				CurTest->AddWarning(FString(V), STACK_OFFSET);
 			}
-		}
-		// Errors
-		else if ( Verbosity == ELogVerbosity::Error )
-		{
-			CurTest->AddError(FString(V), STACK_OFFSET);
-		}
-		// Display
-		if ( Verbosity == ELogVerbosity::Display )
-		{
-			CurTest->AddInfo(FString(V), STACK_OFFSET);
+			// Display
+			else
+			{
+				CurTest->AddInfo(FString(V), STACK_OFFSET);
+			}
 		}
 		// Log...etc
 		else
@@ -673,16 +674,6 @@ void FAutomationTestFramework::AddAnalyticsItemToCurrentTest( const FString& Ana
 	}
 }
 
-bool FAutomationTestFramework::GetTreatWarningsAsErrors() const
-{
-	return AutomationTestFeedbackContext.TreatWarningsAsErrors;
-}
-
-void FAutomationTestFramework::SetTreatWarningsAsErrors(TOptional<bool> bTreatWarningsAsErrors)
-{
-	AutomationTestFeedbackContext.TreatWarningsAsErrors = bTreatWarningsAsErrors.IsSet() ? bTreatWarningsAsErrors.GetValue() : GWarn->TreatWarningsAsErrors;
-}
-
 void FAutomationTestFramework::NotifyScreenshotComparisonComplete(const FAutomationScreenshotCompareResults& CompareResults)
 {
 	OnScreenshotCompared.Broadcast(CompareResults);
@@ -877,7 +868,7 @@ void FAutomationTestBase::ClearExecutionInfo()
 
 void FAutomationTestBase::AddError(const FString& InError, int32 StackOffset)
 {
-	if( !bSuppressLogs && !IsExpectedError(InError))
+	if( !IsExpectedError(InError))
 	{
 		ExecutionInfo.AddEvent(FAutomationEvent(EAutomationEventType::Error, InError), StackOffset + 1);
 	}
@@ -893,7 +884,7 @@ void FAutomationTestBase::AddErrorIfFalse(bool bCondition, const FString& InErro
 
 void FAutomationTestBase::AddErrorS(const FString& InError, const FString& InFilename, int32 InLineNumber)
 {
-	if ( !bSuppressLogs && !IsExpectedError(InError))
+	if ( !IsExpectedError(InError))
 	{
 		//ExecutionInfo.AddEvent(FAutomationEvent(EAutomationEventType::Error, InError, ExecutionInfo.GetContext(), InFilename, InLineNumber));
 	}
@@ -901,7 +892,7 @@ void FAutomationTestBase::AddErrorS(const FString& InError, const FString& InFil
 
 void FAutomationTestBase::AddWarningS(const FString& InWarning, const FString& InFilename, int32 InLineNumber)
 {
-	if ( !bSuppressLogs && !IsExpectedError(InWarning))
+	if ( !IsExpectedError(InWarning))
 	{
 		//ExecutionInfo.AddEvent(FAutomationEvent(EAutomationEventType::Warning, InWarning, ExecutionInfo.GetContext(), InFilename, InLineNumber));
 	}
@@ -909,7 +900,7 @@ void FAutomationTestBase::AddWarningS(const FString& InWarning, const FString& I
 
 void FAutomationTestBase::AddWarning( const FString& InWarning, int32 StackOffset )
 {
-	if ( !bSuppressLogs && !IsExpectedError(InWarning))
+	if ( !IsExpectedError(InWarning))
 	{
 		ExecutionInfo.AddEvent(FAutomationEvent(EAutomationEventType::Warning, InWarning), StackOffset + 1);
 	}
@@ -917,10 +908,7 @@ void FAutomationTestBase::AddWarning( const FString& InWarning, int32 StackOffse
 
 void FAutomationTestBase::AddInfo( const FString& InLogItem, int32 StackOffset )
 {
-	if ( !bSuppressLogs )
-	{
-		ExecutionInfo.AddEvent(FAutomationEvent(EAutomationEventType::Info, InLogItem), StackOffset + 1);
-	}
+	ExecutionInfo.AddEvent(FAutomationEvent(EAutomationEventType::Info, InLogItem), StackOffset + 1);
 }
 
 void FAutomationTestBase::AddAnalyticsItem(const FString& InAnalyticsItem)
@@ -930,10 +918,7 @@ void FAutomationTestBase::AddAnalyticsItem(const FString& InAnalyticsItem)
 
 void FAutomationTestBase::AddEvent(const FAutomationEvent& InEvent, int32 StackOffset)
 {
-	if (!bSuppressLogs)
-	{
-		ExecutionInfo.AddEvent(InEvent, StackOffset + 1);
-	}
+	ExecutionInfo.AddEvent(InEvent, StackOffset + 1);
 }
 
 bool FAutomationTestBase::HasAnyErrors() const
diff --git a/Engine/Source/Runtime/Core/Public/Misc/AutomationTest.h b/Engine/Source/Runtime/Core/Public/Misc/AutomationTest.h
--- a/Engine/Source/Runtime/Core/Public/Misc/AutomationTest.h
+++ b/Engine/Source/Runtime/Core/Public/Misc/AutomationTest.h
@@ -856,9 +856,6 @@ public:
 		return CurrentTest;
 	}
 
-	bool GetTreatWarningsAsErrors() const;
-	void SetTreatWarningsAsErrors(TOptional<bool> bTreatWarningsAsErrors);
-
 	void NotifyScreenshotComparisonComplete(const FAutomationScreenshotCompareResults& CompareResults);
 	void NotifyTestDataRetrieved(bool bWasNew, const FString& JsonData);
 	void NotifyPerformanceDataRetrieved(bool bSuccess, const FString& ErrorMessage);
@@ -1003,7 +1000,6 @@ public:
 	 */
 	FAutomationTestBase( const FString& InName, const bool bInComplexTask )
 		: bComplexTask( bInComplexTask )
-		, bSuppressLogs( false )
 		, TestName( InName )
 	{
 		// Register the newly created automation test into the automation testing framework
@@ -1173,15 +1169,29 @@ public:
 	}
 
 	/**
-	 * Used to suppress / unsuppress logs.
+	 * If true logs will not be included in test events
 	 *
-	 * @param bNewValue - True if you want to suppress logs.  False to unsuppress.
+	 * @return true to suppress logs
 	 */
-	void SetSuppressLogs(bool bNewValue)
+	virtual bool SuppressLogs()
 	{
-		bSuppressLogs = bNewValue;
+		return false;
 	}
 
+	/**
+	 * If true (and SuppressLogs=false) then LogErrors will be treated as test errors
+	 *
+	 * @return true to make errors errors
+	 */
+	virtual bool TreatLogErrorsAsErrors() { return true; }
+
+	/**
+	 * If true (and SuppressLogs=false) then LogWarnings will be treated as test errors
+	 *
+	 * @return true to make warnings errors
+	 */
+	virtual bool TreatLogWarningsAsErrors() { return false; }
+
 	/**
 	 * Enqueues a new latent command.
 	 */
diff --git a/Engine/Source/Runtime/Engine/Classes/Tests/AutomationTestSettings.h b/Engine/Source/Runtime/Engine/Classes/Tests/AutomationTestSettings.h
--- a/Engine/Source/Runtime/Engine/Classes/Tests/AutomationTestSettings.h
+++ b/Engine/Source/Runtime/Engine/Classes/Tests/AutomationTestSettings.h
@@ -331,6 +331,50 @@ struct FLaunchOnTestSettings
 	FString DeviceID;
 };
 
+/*
+* Describes a filter for a test group.
+*/
+USTRUCT()
+struct FAutomatedTestFilter
+{
+	GENERATED_USTRUCT_BODY()
+
+public:
+
+	FAutomatedTestFilter(FString InContains, bool InMatchFromStart = false)
+		: Contains(InContains), MatchFromStart(InMatchFromStart)
+	{
+	}
+
+	FAutomatedTestFilter() : FAutomatedTestFilter(TEXT("")) {}
+
+	/** String that the test must contain */
+	UPROPERTY(Config)
+		FString Contains;
+
+	/** If true start matching from the start of the string, else anywhere */
+	UPROPERTY(Config)
+		bool MatchFromStart;
+};
+
+/*
+ *	Describes a group of tests. Each group has a name and a set of filters that determine group membership
+ */
+USTRUCT()
+struct FAutomatedTestGroup
+{
+	GENERATED_USTRUCT_BODY()
+
+public:
+
+	UPROPERTY(Config)
+		FString Name;
+
+	UPROPERTY(Config)
+		TArray<FAutomatedTestFilter> Filters;
+};
+
+
 /**
  * Implements the Editor's user settings.
  */
@@ -424,4 +468,16 @@ public:
 	 */
 	UPROPERTY(EditAnywhere, config, Category=Screenshots)
 	FIntPoint DefaultScreenshotResolution;
+
+	/** List of user-defined test groups */
+	UPROPERTY(Config)
+	TArray<FAutomatedTestGroup> Groups;
+
+	/** Whether to treat log errors as test errors (default=true) */
+	UPROPERTY(Config)
+	bool bTreatLogErrorsAsTestErrors;
+
+	/** Whether to treat log warnings as test errors (default=false) */
+	UPROPERTY(Config)
+	bool bTreatLogWarningsAsTestErrors;
 };
diff --git a/Engine/Source/Runtime/Engine/Private/Tests/AutomationTestSettings.cpp b/Engine/Source/Runtime/Engine/Private/Tests/AutomationTestSettings.cpp
new file mode 100644
--- /dev/null
+++ b/Engine/Source/Runtime/Engine/Private/Tests/AutomationTestSettings.cpp
@@ -0,0 +1,13 @@
+// Copyright 1998-2019 Epic Games, Inc. All Rights Reserved.
+
+
+#include "Tests/AutomationTestSettings.h"
+
+
+UAutomationTestSettings::UAutomationTestSettings(const FObjectInitializer& ObjectInitializer)
+	: Super(ObjectInitializer)
+	, bTreatLogErrorsAsTestErrors(true)
+	, bTreatLogWarningsAsTestErrors(false)
+{
+	DefaultScreenshotResolution = FIntPoint(1920, 1080);
+}
diff --git a/Engine/Source/Runtime/Engine/Private/Tests/EngineAutomationTests.cpp b/Engine/Source/Runtime/Engine/Private/Tests/EngineAutomationTests.cpp
--- a/Engine/Source/Runtime/Engine/Private/Tests/EngineAutomationTests.cpp
+++ b/Engine/Source/Runtime/Engine/Private/Tests/EngineAutomationTests.cpp
@@ -1191,12 +1191,3 @@ bool FWaitForMatineeToCompleteAndDoScreenshotsLatentCommand::Update()
 //}
 
 #endif //WITH_DEV_AUTOMATION_TESTS
-
-/* UAutomationTestSettings interface
-*****************************************************************************/
-
-UAutomationTestSettings::UAutomationTestSettings(const FObjectInitializer& ObjectInitializer)
-	: Super(ObjectInitializer)
-{
-	DefaultScreenshotResolution = FIntPoint(1920, 1080);
-}
