Moved low level test settings to a place with no engine dependencies.

diff --git a/Engine/Source/Developer/AutomationController/Private/AutomationCommandline.cpp b/Engine/Source/Developer/AutomationController/Private/AutomationCommandline.cpp
--- a/Engine/Source/Developer/AutomationController/Private/AutomationCommandline.cpp
+++ b/Engine/Source/Developer/AutomationController/Private/AutomationCommandline.cpp
@@ -14,7 +14,7 @@
 #include "Misc/Paths.h"
 #include "Misc/FileHelper.h"
 #include "AssetRegistryModule.h"
-#include "Tests/AutomationTestSettings.h"
+#include "AutomationControllerSettings.h"
 
 DEFINE_LOG_CATEGORY_STATIC(LogAutomationCommandLine, Log, All);
 
@@ -133,7 +133,7 @@ public:
 		TArray<FAutomatedTestFilter> Filters;
 
 		// get our settings CDO where things are stored
-		UAutomationTestSettings* Settings = UAutomationTestSettings::StaticClass()->GetDefaultObject<UAutomationTestSettings>();
+		UAutomationControllerSettings* Settings = UAutomationControllerSettings::StaticClass()->GetDefaultObject<UAutomationControllerSettings>();
 
 		// iterate through the arguments to build a filter list by doing the following -
 		// 1) If argument is a filter (filter:system) then make sure we only filter-in tests that start with that filter
diff --git a/Engine/Source/Developer/AutomationController/Private/AutomationControllerSettings.cpp b/Engine/Source/Developer/AutomationController/Private/AutomationControllerSettings.cpp
new file mode 100644
--- /dev/null
+++ b/Engine/Source/Developer/AutomationController/Private/AutomationControllerSettings.cpp
@@ -0,0 +1,12 @@
+// Copyright 1998-2019 Epic Games, Inc. All Rights Reserved.
+
+
+#include "AutomationControllerSettings.h"
+
+
+UAutomationControllerSettings::UAutomationControllerSettings(const FObjectInitializer& ObjectInitializer)
+	: Super(ObjectInitializer)
+	, bTreatLogErrorsAsTestErrors(true)
+	, bTreatLogWarningsAsTestErrors(false)
+{
+}
diff --git a/Engine/Source/Developer/AutomationController/Public/AutomationControllerSettings.h b/Engine/Source/Developer/AutomationController/Public/AutomationControllerSettings.h
new file mode 100644
--- /dev/null
+++ b/Engine/Source/Developer/AutomationController/Public/AutomationControllerSettings.h
@@ -0,0 +1,77 @@
+// Copyright 1998-2019 Epic Games, Inc. All Rights Reserved.
+
+
+#pragma once
+
+#include "CoreMinimal.h"
+#include "UObject/Object.h"
+#include "UObject/SoftObjectPath.h"
+#include "AutomationControllerSettings.generated.h"
+
+/*
+* Describes a filter for a test group.
+*/
+USTRUCT()
+struct FAutomatedTestFilter
+{
+	GENERATED_USTRUCT_BODY()
+
+public:
+
+	FAutomatedTestFilter(FString InContains, bool InMatchFromStart = false)
+		: Contains(InContains), MatchFromStart(InMatchFromStart)
+	{
+	}
+
+	FAutomatedTestFilter() : FAutomatedTestFilter(TEXT("")) {}
+
+	/** String that the test must contain */
+	UPROPERTY(Config)
+		FString Contains;
+
+	/** If true start matching from the start of the string, else anywhere */
+	UPROPERTY(Config)
+		bool MatchFromStart;
+};
+
+/*
+ *	Describes a group of tests. Each group has a name and a set of filters that determine group membership
+ */
+USTRUCT()
+struct FAutomatedTestGroup
+{
+	GENERATED_USTRUCT_BODY()
+
+public:
+
+	UPROPERTY(Config)
+		FString Name;
+
+	UPROPERTY(Config)
+		TArray<FAutomatedTestFilter> Filters;
+};
+
+
+
+/**
+ * Implements the Editor's user settings.
+ */
+UCLASS(config = Engine, defaultconfig)
+class AUTOMATIONCONTROLLER_API UAutomationControllerSettings : public UObject
+{
+	GENERATED_UCLASS_BODY()
+
+public:
+	
+	/** List of user-defined test groups */
+	UPROPERTY(Config)
+	TArray<FAutomatedTestGroup> Groups;
+
+	/** Whether to treat log errors as test errors (default=true) */
+	UPROPERTY(Config)
+	bool bTreatLogErrorsAsTestErrors;
+
+	/** Whether to treat log warnings as test errors (default=false) */
+	UPROPERTY(Config)
+	bool bTreatLogWarningsAsTestErrors;
+};
diff --git a/Engine/Source/Developer/FunctionalTesting/Private/FunctionalTestBase.cpp b/Engine/Source/Developer/FunctionalTesting/Private/FunctionalTestBase.cpp
--- a/Engine/Source/Developer/FunctionalTesting/Private/FunctionalTestBase.cpp
+++ b/Engine/Source/Developer/FunctionalTesting/Private/FunctionalTestBase.cpp
@@ -1,7 +1,7 @@
 // Copyright 1998-2019 Epic Games, Inc. All Rights Reserved.
 
 #include "FunctionalTestBase.h"
-#include "Tests/AutomationTestSettings.h"
+#include "AutomationControllerSettings.h"
 
 
 FFunctionalTestBase::FFunctionalTestBase(const FString& InName, const bool bInComplexTask)
@@ -17,7 +17,7 @@ FFunctionalTestBase::FFunctionalTestBase(const FString& InName, const bool bInCo
 void FFunctionalTestBase::SetLogErrorAndWarningHandlingToDefault()
 {
 	// Set to project defaults
-	UAutomationTestSettings* Settings = UAutomationTestSettings::StaticClass()->GetDefaultObject<UAutomationTestSettings>();
+	UAutomationControllerSettings* Settings = UAutomationControllerSettings::StaticClass()->GetDefaultObject<UAutomationControllerSettings>();
 	bTreatLogErrorsAsErrors = Settings->bTreatLogErrorsAsTestErrors;
 	bTreatLogWarningsAsErrors = Settings->bTreatLogWarningsAsTestErrors;
 }
diff --git a/Engine/Source/Runtime/Engine/Classes/Tests/AutomationTestSettings.h b/Engine/Source/Runtime/Engine/Classes/Tests/AutomationTestSettings.h
--- a/Engine/Source/Runtime/Engine/Classes/Tests/AutomationTestSettings.h
+++ b/Engine/Source/Runtime/Engine/Classes/Tests/AutomationTestSettings.h
@@ -331,48 +331,6 @@ struct FLaunchOnTestSettings
 	FString DeviceID;
 };
 
-/*
-* Describes a filter for a test group.
-*/
-USTRUCT()
-struct FAutomatedTestFilter
-{
-	GENERATED_USTRUCT_BODY()
-
-public:
-
-	FAutomatedTestFilter(FString InContains, bool InMatchFromStart = false)
-		: Contains(InContains), MatchFromStart(InMatchFromStart)
-	{
-	}
-
-	FAutomatedTestFilter() : FAutomatedTestFilter(TEXT("")) {}
-
-	/** String that the test must contain */
-	UPROPERTY(Config)
-		FString Contains;
-
-	/** If true start matching from the start of the string, else anywhere */
-	UPROPERTY(Config)
-		bool MatchFromStart;
-};
-
-/*
- *	Describes a group of tests. Each group has a name and a set of filters that determine group membership
- */
-USTRUCT()
-struct FAutomatedTestGroup
-{
-	GENERATED_USTRUCT_BODY()
-
-public:
-
-	UPROPERTY(Config)
-		FString Name;
-
-	UPROPERTY(Config)
-		TArray<FAutomatedTestFilter> Filters;
-};
 
 
 /**
@@ -468,16 +426,4 @@ public:
 	 */
 	UPROPERTY(EditAnywhere, config, Category=Screenshots)
 	FIntPoint DefaultScreenshotResolution;
-
-	/** List of user-defined test groups */
-	UPROPERTY(Config)
-	TArray<FAutomatedTestGroup> Groups;
-
-	/** Whether to treat log errors as test errors (default=true) */
-	UPROPERTY(Config)
-	bool bTreatLogErrorsAsTestErrors;
-
-	/** Whether to treat log warnings as test errors (default=false) */
-	UPROPERTY(Config)
-	bool bTreatLogWarningsAsTestErrors;
 };
diff --git a/Engine/Source/Runtime/Engine/Private/Tests/AutomationTestSettings.cpp b/Engine/Source/Runtime/Engine/Private/Tests/AutomationTestSettings.cpp
--- a/Engine/Source/Runtime/Engine/Private/Tests/AutomationTestSettings.cpp
+++ b/Engine/Source/Runtime/Engine/Private/Tests/AutomationTestSettings.cpp
@@ -6,8 +6,6 @@
 
 UAutomationTestSettings::UAutomationTestSettings(const FObjectInitializer& ObjectInitializer)
 	: Super(ObjectInitializer)
-	, bTreatLogErrorsAsTestErrors(true)
-	, bTreatLogWarningsAsTestErrors(false)
 {
 	DefaultScreenshotResolution = FIntPoint(1920, 1080);
 }
