Missed file from 6258941

diff --git a/Engine/Source/Developer/AutomationController/AutomationController.Build.cs b/Engine/Source/Developer/AutomationController/AutomationController.Build.cs
--- a/Engine/Source/Developer/AutomationController/AutomationController.Build.cs
+++ b/Engine/Source/Developer/AutomationController/AutomationController.Build.cs
@@ -10,8 +10,7 @@ namespace UnrealBuildTool.Rules
 				new string[]
 				{
 					"Core",
-					"CoreUObject",
-					"Engine"
+					"CoreUObject"
 				}
 			); 
 			
@@ -27,7 +26,7 @@ namespace UnrealBuildTool.Rules
 					"ScreenShotComparisonTools",
 					"HTTP",
                     "AssetRegistry"
-                }
+				}
 			);
 
             if (Target.bBuildEditor)
