Fix uninitialized struct member causing warnings when creating DDC on Mac.

diff --git a/Engine/Source/Developer/AutomationController/Private/AutomationControllerManager.h b/Engine/Source/Developer/AutomationController/Private/AutomationControllerManager.h
--- a/Engine/Source/Developer/AutomationController/Private/AutomationControllerManager.h
+++ b/Engine/Source/Developer/AutomationController/Private/AutomationControllerManager.h
@@ -81,7 +81,8 @@ struct FAutomatedTestPassResults
 
 public:
 	FAutomatedTestPassResults()
-		: Succeeded(0)
+		: ReportCreatedOn(0)
+		, Succeeded(0)
 		, SucceededWithWarnings(0)
 		, Failed(0)
 		, NotRun(0)
