Fixed warnings when workers are not ready yet.
Cut down on number of warnings issued.

diff --git a/Engine/Source/Developer/AutomationController/Private/AutomationCommandline.cpp b/Engine/Source/Developer/AutomationController/Private/AutomationCommandline.cpp
--- a/Engine/Source/Developer/AutomationController/Private/AutomationCommandline.cpp
+++ b/Engine/Source/Developer/AutomationController/Private/AutomationCommandline.cpp
@@ -198,8 +198,21 @@ public:
 		
 		if (AutomationController->GetNumDeviceClusters() == 0)
 		{
-			// Couldn't find any workers, go back into finding workers
-			UE_LOG(LogAutomationCommandLine, Warning, TEXT("Can't find any workers! Searching again"));
+			static double FirstWarningTime = FPlatformTime::Seconds();
+			static int WarningCount = 0;
+
+			double TimeWaiting = FPlatformTime::Seconds() - FirstWarningTime;
+
+			// This can get called a number of times before a worker is ready, so be conservative in how often we warn
+			if ((++WarningCount % 5) == 0 && TimeWaiting > 10.0)
+			{
+				UE_LOG(LogAutomationCommandLine, Log, TEXT("Can't find any workers! Searching again"));
+			}
+			else
+			{
+				UE_LOG(LogAutomationCommandLine, Warning, TEXT("Can't find any workers! Searching again"));
+			}
+
 			AutomationTestState = EAutomationTestState::FindWorkers;
 			return;
 		}
