Add additional logging during EngineTest report generation
#rnx

diff --git a/Engine/Source/Developer/AutomationController/Private/AutomationControllerManger.cpp b/Engine/Source/Developer/AutomationController/Private/AutomationControllerManger.cpp
--- a/Engine/Source/Developer/AutomationController/Private/AutomationControllerManger.cpp
+++ b/Engine/Source/Developer/AutomationController/Private/AutomationControllerManger.cpp
@@ -675,8 +675,11 @@ void FAutomationControllerManager::ProcessResults()
 			IFileManager::Get().Move(*TempDirectory, *ReportOutputPath);
 			IFileManager::Get().DeleteDirectory(*TempDirectory, false, true);
 		}
-				
+
+		UE_LOG(LogAutomationController, Display, TEXT("Exporting Screenshot Comparison Results"));
+		double ExportStartTime = FPlatformTime::Seconds();
 		FScreenshotExportResults ExportResults = ScreenshotManager->ExportComparisonResultsAsync(ReportOutputPath).Get();
+		UE_LOG(LogAutomationController, Display, TEXT("Exported Screenshot Comparison Results in %f seconds"), FPlatformTime::Seconds() - ExportStartTime);
 
 		FAutomatedTestPassResults SerializedPassResults = OurPassResults;
 
@@ -692,6 +695,7 @@ void FAutomationControllerManager::ProcessResults()
 			SerializedPassResults.ComparisonExportDirectory = DisplayReportOutputPath / FString::FromInt(FEngineVersion::Current().GetChangelist());
 		}
 
+		UE_LOG(LogAutomationController, Display, TEXT("Sorting Test Results"));
 		{
 			SerializedPassResults.Tests.StableSort([] (const FAutomatedTestResult& A, const FAutomatedTestResult& B) {
 				if ( A.GetErrorTotal() > 0 )
