Moved artifact copying to ParallelFor and added additional logging to help with test timeout (believed to be slow network issues).

diff --git a/Engine/Source/Developer/AutomationController/Private/AutomationControllerManger.cpp b/Engine/Source/Developer/AutomationController/Private/AutomationControllerManger.cpp
--- a/Engine/Source/Developer/AutomationController/Private/AutomationControllerManger.cpp
+++ b/Engine/Source/Developer/AutomationController/Private/AutomationControllerManger.cpp
@@ -26,6 +26,7 @@
 #if WITH_EDITOR
 #include "Logging/MessageLog.h"
 #endif
+#include "Async/ParallelFor.h"
 
 // these strings are parsed by Gauntlet (AutomationLogParser) so make sure changes are replicated there!
 #define AutomationTestStarting		TEXT("Test Started. Name={%s}")
@@ -660,12 +661,14 @@ void FAutomationControllerManager::ProcessResults()
 
 	if ( !ReportOutputPath.IsEmpty() )
 	{
-		FDateTime Timestamp = FDateTime::Now();
+		FDateTime StartTime = FDateTime::Now();
 
 		UE_LOG(LogAutomationController, Display, TEXT("Generating Automation Report @ %s."), *ReportOutputPath);
 
 		if ( IFileManager::Get().DirectoryExists(*ReportOutputPath) )
 		{
+			FDateTime StepTime = FDateTime::Now();
+
 			UE_LOG(LogAutomationController, Display, TEXT("Existing report directory found, deleting %s."), *ReportOutputPath);
 
 			// Clear the old report folder.  Why move it first?  Because RemoveDirectory
@@ -674,50 +677,63 @@ void FAutomationControllerManager::ProcessResults()
 			FString TempDirectory = FPaths::GetPath(ReportOutputPath) + TEXT("\\") + FGuid::NewGuid().ToString(EGuidFormats::DigitsWithHyphens);
 			IFileManager::Get().Move(*TempDirectory, *ReportOutputPath);
 			IFileManager::Get().DeleteDirectory(*TempDirectory, false, true);
-		}
 
-		UE_LOG(LogAutomationController, Display, TEXT("Exporting Screenshot Comparison Results"));
-		double ExportStartTime = FPlatformTime::Seconds();
-		FScreenshotExportResults ExportResults = ScreenshotManager->ExportComparisonResultsAsync(ReportOutputPath).Get();
-		UE_LOG(LogAutomationController, Display, TEXT("Exported Screenshot Comparison Results in %f seconds"), FPlatformTime::Seconds() - ExportStartTime);
+			UE_LOG(LogAutomationController, Display, TEXT("Deleted directory in %.02f Seconds"), (FDateTime::Now() - StepTime).GetTotalSeconds());
+		}
 
-		FAutomatedTestPassResults SerializedPassResults = OurPassResults;
+		FAutomatedTestPassResults SerializedPassResults;
 
-		SerializedPassResults.ComparisonExported = ExportResults.Success;
-		SerializedPassResults.ComparisonExportDirectory = ExportResults.ExportPath;
-		SerializedPassResults.ReportCreatedOn = Timestamp;
-		if (DisplayReportOutputPath.IsEmpty())
 		{
+			FDateTime StepTime = FDateTime::Now();
+
+			UE_LOG(LogAutomationController, Display, TEXT("Exporting comparison results to %s..."), *ReportOutputPath);
+
+			FScreenshotExportResults ExportResults = ScreenshotManager->ExportComparisonResultsAsync(ReportOutputPath).Get();
+
+			SerializedPassResults = OurPassResults;
+
+			SerializedPassResults.ComparisonExported = ExportResults.Success;
 			SerializedPassResults.ComparisonExportDirectory = ExportResults.ExportPath;
-		}
-		else
-		{
-			SerializedPassResults.ComparisonExportDirectory = DisplayReportOutputPath / FString::FromInt(FEngineVersion::Current().GetChangelist());
+			SerializedPassResults.ReportCreatedOn = StartTime;
+			if (DisplayReportOutputPath.IsEmpty())
+			{
+				SerializedPassResults.ComparisonExportDirectory = ExportResults.ExportPath;
+			}
+			else
+			{
+				SerializedPassResults.ComparisonExportDirectory = DisplayReportOutputPath / FString::FromInt(FEngineVersion::Current().GetChangelist());
+			}
+
+			UE_LOG(LogAutomationController, Display, TEXT("Exported results in %.02f Seconds"), (FDateTime::Now() - StepTime).GetTotalSeconds());
 		}
 
-		UE_LOG(LogAutomationController, Display, TEXT("Sorting Test Results"));
+
 		{
-			SerializedPassResults.Tests.StableSort([] (const FAutomatedTestResult& A, const FAutomatedTestResult& B) {
-				if ( A.GetErrorTotal() > 0 )
+			FDateTime StepTime = FDateTime::Now();
+
+			UE_LOG(LogAutomationController, Display, TEXT("Copying artifacts to %s..."), *ReportOutputPath);
+
+			SerializedPassResults.Tests.StableSort([](const FAutomatedTestResult& A, const FAutomatedTestResult& B) {
+				if (A.GetErrorTotal() > 0)
 				{
-					if ( B.GetErrorTotal() > 0 )
-						return ( A.FullTestPath < B.FullTestPath );
+					if (B.GetErrorTotal() > 0)
+						return (A.FullTestPath < B.FullTestPath);
 					else
 						return true;
 				}
-				else if ( B.GetErrorTotal() > 0 )
+				else if (B.GetErrorTotal() > 0)
 				{
 					return false;
 				}
 
-				if ( A.GetWarningTotal() > 0 )
+				if (A.GetWarningTotal() > 0)
 				{
-					if ( B.GetWarningTotal() > 0 )
-						return ( A.FullTestPath < B.FullTestPath );
+					if (B.GetWarningTotal() > 0)
+						return (A.FullTestPath < B.FullTestPath);
 					else
 						return true;
 				}
-				else if ( B.GetWarningTotal() > 0 )
+				else if (B.GetWarningTotal() > 0)
 				{
 					return false;
 				}
@@ -725,34 +741,79 @@ void FAutomationControllerManager::ProcessResults()
 				return A.FullTestPath < B.FullTestPath;
 			});
 
-			for ( FAutomatedTestResult& Test : SerializedPassResults.Tests )
+			// used for reporting and sync during copies
+			FCriticalSection CS;
+			int TotalFileArtifacts = 0;
+			int CopiedFileArtifacts = 0;
+
+			// Get a total for reporting
+			for (FAutomatedTestResult& Test : SerializedPassResults.Tests)
+			{
+				for (FAutomationArtifact& Artifact : Test.GetArtifacts())
+				{
+					TotalFileArtifacts += Artifact.LocalFiles.Num();
+				}
+			}
+
+
+			// most tests have a single set of artifacts with three or more files, so we could optimize this further by gathering them all first...
+			for (FAutomatedTestResult& Test : SerializedPassResults.Tests)
 			{
-				for (FAutomationArtifact& Artifact : Test.GetArtifacts() )
+				int TestArtifactCount = 0;
+
+				for (FAutomationArtifact& Artifact : Test.GetArtifacts())
 				{
-					for ( const auto& Entry : Artifact.LocalFiles )
+					TArray<FString> Keys;
+					Artifact.LocalFiles.GetKeys(Keys);
+
+					ParallelFor(Keys.Num(), [&](int32 Index)
 					{
-						Artifact.Files.Add(Entry.Key, CopyArtifact(ReportOutputPath, Entry.Value));
-					}
+						const FString& Key = Keys[Index];
+
+						FString Path = CopyArtifact(ReportOutputPath, Artifact.LocalFiles[Key]);
+						{
+							FScopeLock Lock(&CS);
+							Artifact.Files.Add(Key, MoveTemp(Path));
+							CopiedFileArtifacts++;
+							TestArtifactCount++;
+
+							// Show occasional progress for larger result sets
+							if ((CopiedFileArtifacts % 50) == 0)
+							{
+								UE_LOG(LogAutomationController, Display, TEXT("Copied %d of %d files in %.02f Seconds"), CopiedFileArtifacts, TotalFileArtifacts, (FDateTime::Now() - StepTime).GetTotalSeconds());
+							}
+						}
+					});
+
+					//UE_LOG(LogAutomationController, Verbose, TEXT("Copied %d files from artifact %s"), Keys.Num(), *Artifact.Name);
 				}
+
+				//UE_LOG(LogAutomationController, Verbose, TEXT("Copied %d files from test %s"), TestArtifactCount, *Test.TestDisplayName);
 			}
+
+			UE_LOG(LogAutomationController, Display, TEXT("Copied %d files in %.02f Seconds"), TotalFileArtifacts, (FDateTime::Now() - StepTime).GetTotalSeconds());
 		}
 
-		UE_LOG(LogAutomationController, Display, TEXT("Writing reports... %s."), *ReportOutputPath);
+		{
+			FDateTime StepTime = FDateTime::Now();
+			UE_LOG(LogAutomationController, Display, TEXT("Writing reports to %s..."), *ReportOutputPath);
 
-		// Generate Json
-		GenerateJsonTestPassSummary(SerializedPassResults, Timestamp);
+			// Generate Json
+			GenerateJsonTestPassSummary(SerializedPassResults, StartTime);
 
-		// Generate Html
-		GenerateHtmlTestPassSummary(SerializedPassResults, Timestamp);
+			// Generate Html
+			GenerateHtmlTestPassSummary(SerializedPassResults, StartTime);
 
-		if ( !DeveloperReportUrl.IsEmpty() )
-		{
-			UE_LOG(LogAutomationController, Display, TEXT("Launching Report URL %s."), *DeveloperReportUrl);
+			if (!DeveloperReportUrl.IsEmpty())
+			{
+				UE_LOG(LogAutomationController, Display, TEXT("Launching Report URL %s."), *DeveloperReportUrl);
 
-			FPlatformProcess::LaunchURL(*DeveloperReportUrl, nullptr, nullptr);
+				FPlatformProcess::LaunchURL(*DeveloperReportUrl, nullptr, nullptr);
+			}
+			UE_LOG(LogAutomationController, Display, TEXT("Wrote reports in %.02f Seconds"), (FDateTime::Now() - StepTime).GetTotalSeconds());
 		}
 
-		UE_LOG(LogAutomationController, Display, TEXT("Done writing reports... %s."), *ReportOutputPath);
+		UE_LOG(LogAutomationController, Display, TEXT("Completed exporting all results in %.02f Seconds"), (FDateTime::Now() - StartTime).GetTotalSeconds());
 	}
 
 	// Then clean our array for the next pass.
