Allow filters for automated tests to be on multiple lines

diff --git a/Engine/Source/Developer/AutomationController/Private/AutomationCommandline.cpp b/Engine/Source/Developer/AutomationController/Private/AutomationCommandline.cpp
--- a/Engine/Source/Developer/AutomationController/Private/AutomationCommandline.cpp
+++ b/Engine/Source/Developer/AutomationController/Private/AutomationCommandline.cpp
@@ -163,26 +163,30 @@ public:
 				// if the argument is a group (e.g. Group:Rendering) then seach our groups for one that matches
 				FString GroupName = ArgumentName.RightChop(GroupPrefix.Len());
 
-				FAutomatedTestGroup* GroupEntry = Settings->Groups.FindByPredicate([&GroupName](const FAutomatedTestGroup& E) {
-					return E.Name == GroupName;
-				});
+				bool FoundGroup = false;
 
-				if (GroupEntry)
+				for (int32 i = 0; i < Settings->Groups.Num(); ++i)
 				{
-					// if found add all its filters to our current list
-					if (GroupEntry->Filters.Num() > 0)
+					FAutomatedTestGroup* GroupEntry = &(Settings->Groups[i]);
+					if (GroupEntry && GroupEntry->Name == GroupName)
 					{
-						for (const FAutomatedTestFilter& GroupFilter : GroupEntry->Filters)
+						FoundGroup = true;
+						// if found add all this groups filters to our current list
+						if (GroupEntry->Filters.Num() > 0)
 						{
-							Filters.Add(GroupFilter);
+							for (const FAutomatedTestFilter& GroupFilter : GroupEntry->Filters)
+							{
+								Filters.Add(GroupFilter);
+							}
+						}
+						else
+						{
+							UE_LOG(LogAutomationCommandLine, Warning, TEXT("Group %s contains no filters"), *GroupName);
 						}
-					}
-					else
-					{
-						UE_LOG(LogAutomationCommandLine, Warning, TEXT("Group %s contains no filters"), *GroupName);
 					}
 				}
-				else
+
+				if (!FoundGroup)
 				{
 					UE_LOG(LogAutomationCommandLine, Warning, TEXT("No matching group named %s"), *GroupName);
 				}
