Fix for automation results being added to an anonymous category in the editors log window.

Also tidied up output formatiing a little.

diff --git a/Engine/Source/Developer/AutomationController/Private/AutomationControllerManger.cpp b/Engine/Source/Developer/AutomationController/Private/AutomationControllerManger.cpp
--- a/Engine/Source/Developer/AutomationController/Private/AutomationControllerManger.cpp
+++ b/Engine/Source/Developer/AutomationController/Private/AutomationControllerManger.cpp
@@ -38,6 +38,8 @@
 
 DEFINE_LOG_CATEGORY_STATIC(LogAutomationController, Log, All)
 
+#define LOCTEXT_NAMESPACE "AutomationTesting"
+
 FAutomationControllerManager::FAutomationControllerManager()
 {
 	CheckpointFile = nullptr;
@@ -125,12 +127,12 @@ void FAutomationControllerManager::RunTests(const bool bInIsLocalSession)
 	CheckTestTimer = 0.f;
 
 #if WITH_EDITOR
-	FMessageLog AutomationTestingLog("AutomationTestingLog");
+	FMessageLog AutomationEditorLog("AutomationTestingLog");
 	FString NewPageName = FString::Printf(TEXT("-----Test Run %d----"), ExecutionCount);
 	FText NewPageNameText = FText::FromString(*NewPageName);
-	AutomationTestingLog.Open();
-	AutomationTestingLog.NewPage(NewPageNameText);
-	AutomationTestingLog.Info(NewPageNameText);
+	AutomationEditorLog.Open();
+	AutomationEditorLog.NewPage(NewPageNameText);
+	AutomationEditorLog.Info(NewPageNameText);
 #endif
 	//reset all tests
 	ReportManager.ResetForExecution(NumTestPasses);
@@ -1157,29 +1159,41 @@ void FAutomationControllerManager::ReportAutomationResult(const TSharedPtr<IAuto
 
 	FName CategoryName = "LogAutomationController";
 #if WITH_EDITOR
-	FMessageLog AutomationTestingLog(CategoryName);
+	FMessageLog AutomationEditorLog("AutomationTestingLog");
 	// we log these messages ourselves for non-editor platforms so suppress this.
-	AutomationTestingLog.SuppressLoggingToOutputLog(true);
-	AutomationTestingLog.Open();
+	AutomationEditorLog.SuppressLoggingToOutputLog(true);
+	AutomationEditorLog.Open();
 #endif
 
 	const FAutomationTestResults& Results = InReport->GetResults(ClusterIndex, PassIndex);
 
+	// write results to editor panel
+#if WITH_EDITOR
+	FFormatNamedArguments Arguments;
+	Arguments.Add(TEXT("TestName"), FText::FromString(InReport->GetFullTestPath()));
+	Arguments.Add(TEXT("Result"), Results.State != EAutomationState::Success ? LOCTEXT("Failed", "Failed") : LOCTEXT("Passed", "Passed"));
+	TSharedRef<FTextToken> Token = FTextToken::Create(FText::Format(LOCTEXT("TestSuccessOrFailure", "Test '{TestName}' completed with result '{Result}'"), Arguments));
+	
+	if (Results.State == EAutomationState::Success)
+	{
+		AutomationEditorLog.Info()->AddToken(Token);
+	}
+	else
+	{
+		AutomationEditorLog.Error()->AddToken(Token);
+	}
+#endif
+
+	// Now log
 	if (Results.State == EAutomationState::Success)
 	{
 		FString SuccessString = FString::Printf(AutomationSuccessFormat, *InReport->GetDisplayName(), *InReport->GetFullTestPath());
 		UE_LOG(LogAutomationController, Display, TEXT("%s"), *SuccessString);
-#if WITH_EDITOR
-		AutomationTestingLog.Info(FText::FromString(*SuccessString));
-#endif
 	}
 	else
 	{
 		FString FailureString = FString::Printf(AutomationFailureFormat, *InReport->GetDisplayName(), *InReport->GetFullTestPath());
 		UE_LOG(LogAutomationController, Error, TEXT("%s"), *FailureString);
-#if WITH_EDITOR
-		AutomationTestingLog.Error(FText::FromString(*FailureString));
-#endif
 	}
 
 	// bracket these for easy parsing
@@ -1192,19 +1206,19 @@ void FAutomationControllerManager::ReportAutomationResult(const TSharedPtr<IAuto
 		case EAutomationEventType::Info:
 			UE_LOG(LogAutomationController, Log, TEXT("%s"), *Entry.ToString());
 #if WITH_EDITOR
-			AutomationTestingLog.Info(FText::FromString(Entry.ToString()));
+			AutomationEditorLog.Info(FText::FromString(Entry.ToString()));
 #endif
 			break;
 		case EAutomationEventType::Warning:
 			UE_LOG(LogAutomationController, Warning, TEXT("%s"), *Entry.ToString());
 #if WITH_EDITOR
-			AutomationTestingLog.Warning(FText::FromString(Entry.ToString()));
+			AutomationEditorLog.Warning(FText::FromString(Entry.ToString()));
 #endif
 			break;
 		case EAutomationEventType::Error:
 			UE_LOG(LogAutomationController, Error, TEXT("%s"), *Entry.ToString());
 #if WITH_EDITOR
-			AutomationTestingLog.Error(FText::FromString(Entry.ToString()));
+			AutomationEditorLog.Error(FText::FromString(Entry.ToString()));
 #endif
 			break;
 		}
@@ -1358,3 +1372,5 @@ void FAutomationControllerManager::ResetAutomationTestTimeout(const TCHAR* Reaso
 	//GLog->Logf(ELogVerbosity::Display, TEXT("Resetting automation test timeout: %s"), Reason);
 	LastTimeUpdateTicked = FPlatformTime::Seconds();
 }
+
+#undef LOCTEXT_NAMESPACE
\ No newline at end of file
