Configuration of AutomationController timings
#rnx

diff --git a/Engine/Source/Developer/AutomationController/Private/AutomationControllerManager.h b/Engine/Source/Developer/AutomationController/Private/AutomationControllerManager.h
--- a/Engine/Source/Developer/AutomationController/Private/AutomationControllerManager.h
+++ b/Engine/Source/Developer/AutomationController/Private/AutomationControllerManager.h
@@ -486,6 +486,12 @@ private:
 	/** Timer to keep track of the last time tests were updated */
 	double CheckTestTimer;
 
+	/** The time to wait between test updates in seconds */
+	float CheckTestIntervalSeconds = 1.0f;
+	
+	/** The time to wait before considering a game instance as lost in seconds */
+	float GameInstanceLostTimerSeconds = 300.0f;
+
 	/** Whether tick is still executing tests for different clusters */
 	uint32 ClusterDistributionMask;
 
diff --git a/Engine/Source/Developer/AutomationController/Private/AutomationControllerManger.cpp b/Engine/Source/Developer/AutomationController/Private/AutomationControllerManger.cpp
--- a/Engine/Source/Developer/AutomationController/Private/AutomationControllerManger.cpp
+++ b/Engine/Source/Developer/AutomationController/Private/AutomationControllerManger.cpp
@@ -16,6 +16,7 @@
 #include "AssetEditorMessages.h"
 #include "ImageComparer.h"
 #include "AutomationControllerManager.h"
+#include "AutomationControllerSettings.h"
 #include "Interfaces/IScreenShotToolsModule.h"
 #include "Serialization/JsonSerializer.h"
 #include "JsonObjectConverter.h"
@@ -42,6 +43,19 @@ DEFINE_LOG_CATEGORY_STATIC(LogAutomationController, Log, All)
 
 FAutomationControllerManager::FAutomationControllerManager()
 {
+
+	UAutomationControllerSettings* Settings = UAutomationControllerSettings::StaticClass()->GetDefaultObject<UAutomationControllerSettings>();
+
+	if (Settings->CheckTestIntervalSeconds > 0.0f)
+	{
+		CheckTestIntervalSeconds = Settings->CheckTestIntervalSeconds;
+	}
+	
+	if (Settings->GameInstanceLostTimerSeconds > 0.0f)
+	{
+		GameInstanceLostTimerSeconds = Settings->GameInstanceLostTimerSeconds;
+	}
+	
 	CheckpointFile = nullptr;
 
 	if ( !FParse::Value(FCommandLine::Get(), TEXT("ReportOutputPath="), ReportOutputPath, false) )
@@ -894,18 +908,15 @@ void FAutomationControllerManager::AddPingResult(const FMessageAddress& Responde
 
 void FAutomationControllerManager::UpdateTests()
 {
-	static const float CheckTestInterval = 1.0f;
-	static const float GameInstanceLostTimer = 300.0f;
-
 	CheckTestTimer += FPlatformTime::Seconds() - LastTimeUpdateTicked;
 	LastTimeUpdateTicked = FPlatformTime::Seconds();
-	if ( CheckTestTimer > CheckTestInterval )
+	if (CheckTestTimer > CheckTestIntervalSeconds)
 	{
 		for ( int32 Index = 0; Index < TestRunningArray.Num(); Index++ )
 		{
 			TestRunningArray[Index].LastPingTime += CheckTestTimer;
 
-			if ( TestRunningArray[Index].LastPingTime > GameInstanceLostTimer )
+			if (TestRunningArray[Index].LastPingTime > GameInstanceLostTimerSeconds)
 			{
 				// Find the game session instance info
 				int32 ClusterIndex;
diff --git a/Engine/Source/Developer/AutomationController/Public/AutomationControllerSettings.h b/Engine/Source/Developer/AutomationController/Public/AutomationControllerSettings.h
--- a/Engine/Source/Developer/AutomationController/Public/AutomationControllerSettings.h
+++ b/Engine/Source/Developer/AutomationController/Public/AutomationControllerSettings.h
@@ -74,4 +74,12 @@ public:
 	/** Whether to treat log warnings as test errors (default=false) */
 	UPROPERTY(Config)
 	bool bTreatLogWarningsAsTestErrors;
+	
+	/** How long to wait between test updates (default=1sec)*/
+	UPROPERTY(Config)
+	float CheckTestIntervalSeconds;
+	
+	/** The maximum response wait time for detecting a lost game instance (default=300sec)*/
+	UPROPERTY(Config)
+	float GameInstanceLostTimerSeconds;
 };
