Fix isue with automated test command line delegate

diff --git a/Engine/Source/Developer/AutomationController/Private/AutomationCommandline.cpp b/Engine/Source/Developer/AutomationController/Private/AutomationCommandline.cpp
--- a/Engine/Source/Developer/AutomationController/Private/AutomationCommandline.cpp
+++ b/Engine/Source/Developer/AutomationController/Private/AutomationCommandline.cpp
@@ -1,4 +1,4 @@
-// Copyright Epic Games, Inc. All Rights Reserved.
+// Copyright 1998-2019 Epic Games, Inc. All Rights Reserved.
 
 #include "CoreMinimal.h"
 #include "Misc/CoreMisc.h"
@@ -71,7 +71,9 @@ public:
 		const bool bSendAnalytics = FParse::Param(FCommandLine::Get(), TEXT("SendAutomationAnalytics"));
 
 		// Register for the callback that tells us there are tests available
-		AutomationController->OnTestsRefreshed().AddRaw(this, &FAutomationExecCmd::HandleRefreshTestCallback);
+		if (!TestsRefreshedHandle.IsValid()) {
+			TestsRefreshedHandle = AutomationController->OnTestsRefreshed().AddRaw(this, &FAutomationExecCmd::HandleRefreshTestCallback);
+		}
 
 		TickHandler = FTicker::GetCoreTicker().AddTicker(FTickerDelegate::CreateRaw(this, &FAutomationExecCmd::Tick));
 
@@ -328,6 +330,12 @@ public:
 			{
 				AutomationTestState = EAutomationTestState::Complete;
 			}
+
+
+			// Clear delegate to avoid re-running tests due to multiple delegates being added or when refreshing session frontend
+			// The delegate will be readded in Init whenever a new command is executed
+			AutomationController->OnTestsRefreshed().Remove(TestsRefreshedHandle);
+			TestsRefreshedHandle.Reset();
 		}
 		else if (AutomationCommand == EAutomationCommand::RunCheckpointTests)
 		{
@@ -662,6 +670,9 @@ private:
 	/** What work was requested */
 	EAutomationCommand AutomationCommand;
 
+	/** Handle to Test Refresh delegate */
+	FDelegateHandle TestsRefreshedHandle;
+
 	/** Delay used before finding workers on game instances. Just to ensure they have started up */
 	float DelayTimer;
 
