Add initializers for AutomationControllerManager
#rnx

diff --git a/Engine/Source/Developer/AutomationController/Private/AutomationControllerManager.h b/Engine/Source/Developer/AutomationController/Private/AutomationControllerManager.h
--- a/Engine/Source/Developer/AutomationController/Private/AutomationControllerManager.h
+++ b/Engine/Source/Developer/AutomationController/Private/AutomationControllerManager.h
@@ -460,28 +460,28 @@ private:
 	EAutomationControllerModuleState::Type AutomationTestState;
 
 	/** Which grouping flags are enabled */
-	uint32 DeviceGroupFlags;
+	uint32 DeviceGroupFlags = 0;
 
 	/** Whether to include developer content in the automation tests */
-	bool bDeveloperDirectoryIncluded;
+	bool bDeveloperDirectoryIncluded = false;
 
 	/** Some tests have errors */
-	bool bHasErrors;
+	bool bHasErrors = false;
 
 	/** Some tests have warnings */
-	bool bHasWarning;
+	bool bHasWarning = false;
 
 	/** Some tests have logs */
-	bool bHasLogs;
+	bool bHasLogs = false;
 
 	/** Is this a local session */
 	bool bIsLocalSession;
 
 	/** Are tests results available */
-	bool bTestResultsAvailable;
+	bool bTestResultsAvailable = false;
 
 	/** Which sets of tests to consider */
-	uint32 RequestedTestFlags;
+	uint32 RequestedTestFlags = 0;
 
 	/** Timer to keep track of the last time tests were updated */
 	double CheckTestTimer;
@@ -499,16 +499,16 @@ private:
 	FAutomationDeviceClusterManager DeviceClusterManager;
 
 	/** The iteration number of executing the tests.  Ensures restarting the tests won't allow stale results to try and commit */
-	uint32 ExecutionCount;
+	uint32 ExecutionCount = 0;
 
 	/** Last time the update tests function was ticked */
-	double LastTimeUpdateTicked;
+	double LastTimeUpdateTicked = 0;
 
 	/** Holds the messaging endpoint. */
 	TSharedPtr<FMessageEndpoint, ESPMode::ThreadSafe> MessageEndpoint;
 
 	/** Counter for number of workers we have received responses from for Refreshing the Test List */
-	uint32 RefreshTestResponses;
+	uint32 RefreshTestResponses = 0;
 
 	/** Available stats/status for all tests. */
 	FAutomationReportManager ReportManager;
@@ -531,13 +531,13 @@ private:
 	TArray< FTestRunningInfo > TestRunningArray;
 
 	/** The number of test passes to perform. */
-	int32 NumTestPasses;
+	int32 NumTestPasses = 0;
 
 	/** The current test pass we are on. */
-	int32 CurrentTestPass;
+	int32 CurrentTestPass = 0;
 
 	/** If we should send result to analytics */
-	bool bSendAnalytics;
+	bool bSendAnalytics = false;
 
 	/** The list of results generated by our test pass. */
 	FAutomatedTestPassResults OurPassResults;
