Added improved logging to diagnose failures in saving screenshots

diff --git a/Engine/Source/Developer/AutomationController/Private/AutomationControllerManger.cpp b/Engine/Source/Developer/AutomationController/Private/AutomationControllerManger.cpp
--- a/Engine/Source/Developer/AutomationController/Private/AutomationControllerManger.cpp
+++ b/Engine/Source/Developer/AutomationController/Private/AutomationControllerManger.cpp
@@ -1047,7 +1047,10 @@ void FAutomationControllerManager::HandleReceivedScreenShot(const FAutomationWor
 	
 	if (!FFileHelper::SaveArrayToFile(Message.ScreenImage, *FileName))
 	{
-		UE_LOG(LogAutomationController, Error, TEXT("Failed to save screenshot to %s"), *FileName);
+		uint32 WriteErrorCode = FPlatformMisc::GetLastError();
+		TCHAR WriteErrorBuffer[2048];
+		FPlatformMisc::GetSystemErrorMessage(WriteErrorBuffer, 2048, WriteErrorCode);
+		UE_LOG(LogAutomationController, Warning, TEXT("Fail to save screenshot to %s. WriteError: %u (%s)"), *FileName, WriteErrorCode, WriteErrorBuffer);
 		return;
 	}
 
diff --git a/Engine/Source/Runtime/AutomationWorker/Private/AutomationWorkerModule.cpp b/Engine/Source/Runtime/AutomationWorker/Private/AutomationWorkerModule.cpp
--- a/Engine/Source/Runtime/AutomationWorker/Private/AutomationWorkerModule.cpp
+++ b/Engine/Source/Runtime/AutomationWorker/Private/AutomationWorkerModule.cpp
@@ -454,7 +454,10 @@ void FAutomationWorkerModule::HandleScreenShotCapturedWithName(const TArray<FCol
 
 		if (!FFileHelper::SaveArrayToFile(CompressedBitmap, *LocalFile))
 		{
-			UE_LOG(LogAutomationWorker, Error, TEXT("Failed to save screenshot to %s"), *LocalFile);
+			uint32 WriteErrorCode = FPlatformMisc::GetLastError();
+			TCHAR WriteErrorBuffer[2048];
+			FPlatformMisc::GetSystemErrorMessage(WriteErrorBuffer, 2048, WriteErrorCode);
+			UE_LOG(LogAutomationWorker, Warning, TEXT("Fail to save screenshot to %s. WriteError: %u (%s)"), *LocalFile, WriteErrorCode, WriteErrorBuffer);
 			return;
 		}
 
