Remove unused private fields - Automation

diff --git a/Engine/Source/Developer/AutomationController/Private/AutomationCommandline.cpp b/Engine/Source/Developer/AutomationController/Private/AutomationCommandline.cpp
--- a/Engine/Source/Developer/AutomationController/Private/AutomationCommandline.cpp
+++ b/Engine/Source/Developer/AutomationController/Private/AutomationCommandline.cpp
@@ -707,9 +707,6 @@ private:
 	/** The current state of the automation process */
 	EAutomationTestState AutomationTestState;
 
-	/** The priority flags we would like to run */
-	EAutomationTestFlags::Type AutomationPriority;
-
 	/** What work was requested */
 	TArray<EAutomationCommand> AutomationCommandQueue;
 
@@ -740,9 +737,6 @@ private:
 	//Dictionary that maps flag names to flag values.
 	TMap<FString, int32> FilterMaps;
 
-	//Test pass checkpoint backup file.
-	FArchive* CheckpointFile;
-
 	FString CheckpointCommand;
 
 	TArray<FString> TestsRun;
diff --git a/Engine/Source/Developer/AutomationController/Private/AutomationReport.h b/Engine/Source/Developer/AutomationController/Private/AutomationReport.h
--- a/Engine/Source/Developer/AutomationController/Private/AutomationReport.h
+++ b/Engine/Source/Developer/AutomationController/Private/AutomationReport.h
@@ -88,9 +88,6 @@ private:
 	/** Number of responses from network commands */
 	uint32 NumberNetworkResponsesReceived;
 
-	/** Number of required devices for this test */
-	uint32 RequiredDeviceCount;
-	
 	/** All child tests */
 	TArray<TSharedPtr<IAutomationReport> >ChildReports;
 
diff --git a/Engine/Source/Developer/FunctionalTesting/Private/AutomationBlueprintFunctionLibrary.cpp b/Engine/Source/Developer/FunctionalTesting/Private/AutomationBlueprintFunctionLibrary.cpp
--- a/Engine/Source/Developer/FunctionalTesting/Private/AutomationBlueprintFunctionLibrary.cpp
+++ b/Engine/Source/Developer/FunctionalTesting/Private/AutomationBlueprintFunctionLibrary.cpp
@@ -151,7 +151,6 @@ public:
 		: FSceneViewExtensionBase(AutoRegister)
 		, WorldPtr(InWorld)
 		, Options(InOptions)
-		, CurrentTime(InCurrentTimeToSimulate)
 	{
 	}
 	
@@ -246,7 +245,6 @@ public:
 private:
 	TWeakObjectPtr<UWorld> WorldPtr;
 	FAutomationScreenshotOptions Options;
-	float CurrentTime;
 };
 
 FAutomationTestScreenshotEnvSetup::FAutomationTestScreenshotEnvSetup()
diff --git a/Engine/Source/Runtime/Core/Private/ProfilingDebugging/ABTesting.cpp b/Engine/Source/Runtime/Core/Private/ProfilingDebugging/ABTesting.cpp
--- a/Engine/Source/Runtime/Core/Private/ProfilingDebugging/ABTesting.cpp
+++ b/Engine/Source/Runtime/Core/Private/ProfilingDebugging/ABTesting.cpp
@@ -57,7 +57,6 @@ FABTest::FABTest()
 		, CurrentTest(0)
 		, RemainingTrial(0)
 		, RemainingPrint(0)
-		, SampleIndex(0)
 		, HistoryNum(0)
 		, ReportNum(0)
 		, CoolDown(0)
diff --git a/Engine/Source/Runtime/Core/Public/ProfilingDebugging/ABTesting.h b/Engine/Source/Runtime/Core/Public/ProfilingDebugging/ABTesting.h
--- a/Engine/Source/Runtime/Core/Public/ProfilingDebugging/ABTesting.h
+++ b/Engine/Source/Runtime/Core/Public/ProfilingDebugging/ABTesting.h
@@ -78,7 +78,6 @@ private:
 	int32 CurrentTest;
 	int32 RemainingTrial;
 	int32 RemainingPrint;
-	int32 SampleIndex;
 
 	int32 HistoryNum;
 	int32 ReportNum;
