Add support for telemetry in automated tests

Telemetry data is stored in a csv file that can be picked up later.
Configuration options can be used to specify a custom path and force the csv file to be reset on new (editor) session.

Telemetry data stores: platform, config, datetime, testname, context, datapoint, measurement

diff --git a/Engine/Source/Developer/AutomationController/Private/AutomationControllerManager.h b/Engine/Source/Developer/AutomationController/Private/AutomationControllerManager.h
--- a/Engine/Source/Developer/AutomationController/Private/AutomationControllerManager.h
+++ b/Engine/Source/Developer/AutomationController/Private/AutomationControllerManager.h
@@ -449,6 +449,9 @@ private:
 	/** Handles FAutomationWorkerWorkerOffline messages. */
 	void HandleWorkerOfflineMessage( const FAutomationWorkerWorkerOffline& Message, const TSharedRef<IMessageContext, ESPMode::ThreadSafe>& Context );
 
+	/** Handles FAutomationWorkerTelemetryData messages. */
+	void HandleReceivedTelemetryData(const FAutomationWorkerTelemetryData& Message, const TSharedRef<IMessageContext, ESPMode::ThreadSafe>& Context);
+
 	/** Writes out this automation result to the log */
 	void ReportAutomationResult(const TSharedPtr<IAutomationReport> InReport, int32 ClusterIndex, int32 PassIndex);
 private:
diff --git a/Engine/Source/Developer/AutomationController/Private/AutomationControllerManger.cpp b/Engine/Source/Developer/AutomationController/Private/AutomationControllerManger.cpp
--- a/Engine/Source/Developer/AutomationController/Private/AutomationControllerManger.cpp
+++ b/Engine/Source/Developer/AutomationController/Private/AutomationControllerManger.cpp
@@ -24,6 +24,8 @@
 #include "Misc/FileHelper.h"
 #include "PlatformHttp.h"
 
+#include "AutomationTelemetry.h"
+
 #if WITH_EDITOR
 #include "Logging/MessageLog.h"
 #endif
@@ -638,7 +640,8 @@ void FAutomationControllerManager::Startup()
 		.Handling<FAutomationWorkerRunTestsReply>(this, &FAutomationControllerManager::HandleRunTestsReplyMessage)
 		.Handling<FAutomationWorkerScreenImage>(this, &FAutomationControllerManager::HandleReceivedScreenShot)
 		.Handling<FAutomationWorkerTestDataRequest>(this, &FAutomationControllerManager::HandleTestDataRequest)
-		.Handling<FAutomationWorkerWorkerOffline>(this, &FAutomationControllerManager::HandleWorkerOfflineMessage);
+		.Handling<FAutomationWorkerWorkerOffline>(this, &FAutomationControllerManager::HandleWorkerOfflineMessage)
+		.Handling<FAutomationWorkerTelemetryData>(this, &FAutomationControllerManager::HandleReceivedTelemetryData);
 
 	if ( MessageEndpoint.IsValid() )
 	{
@@ -1268,6 +1271,10 @@ void FAutomationControllerManager::HandleRequestTestsReplyCompleteMessage(const
 	SetTestNames(Context->GetSender(), TestInfo);
 }
 
+void FAutomationControllerManager::HandleReceivedTelemetryData(const FAutomationWorkerTelemetryData& Message, const TSharedRef<IMessageContext, ESPMode::ThreadSafe>& Context)
+{
+	FAutomationTelemetry::HandleAddTelemetry(Message);
+}
 
 void FAutomationControllerManager::ReportAutomationResult(const TSharedPtr<IAutomationReport> InReport, int32 ClusterIndex, int32 PassIndex)
 {
diff --git a/Engine/Source/Developer/AutomationController/Private/AutomationControllerSettings.cpp b/Engine/Source/Developer/AutomationController/Private/AutomationControllerSettings.cpp
--- a/Engine/Source/Developer/AutomationController/Private/AutomationControllerSettings.cpp
+++ b/Engine/Source/Developer/AutomationController/Private/AutomationControllerSettings.cpp
@@ -11,5 +11,6 @@ UAutomationControllerSettings::UAutomationControllerSettings(const FObjectInitia
 	, bTreatLogErrorsAsTestErrors(true)
 	, bTreatLogWarningsAsTestErrors(true)
 	, GameInstanceLostTimerSeconds(300.0f)
+	, bResetTelemetryStorageOnNewSession(false)
 {
 }
diff --git a/Engine/Source/Developer/AutomationController/Private/AutomationTelemetry.cpp b/Engine/Source/Developer/AutomationController/Private/AutomationTelemetry.cpp
new file mode 100644
--- /dev/null
+++ b/Engine/Source/Developer/AutomationController/Private/AutomationTelemetry.cpp
@@ -0,0 +1,184 @@
+// Copyright Epic Games, Inc. All Rights Reserved.
+
+#include "AutomationTelemetry.h"
+
+#include "HAL/FileManager.h"
+#include "Misc/FileHelper.h"
+#include "Misc/Paths.h"
+#include "Misc/DateTime.h"
+#include "AutomationControllerSettings.h"
+
+DEFINE_LOG_CATEGORY_STATIC(LogAutomationTelemetry, Log, All)
+
+bool FAutomationTelemetry::bIsInitialized;
+FString FAutomationTelemetry::TelemetryDirectory;
+bool FAutomationTelemetry::bResetTelemetryStorageOnNewSession;
+
+bool CreateDirectoryRecursively(const FString& FilePath)
+{
+	FString DirectoryPath = FPaths::GetPath(FilePath);
+	if (FPaths::IsDrive(DirectoryPath))
+	{
+		return true;
+	}
+
+	IFileManager& GFileManager = IFileManager::Get();
+	if (!GFileManager.DirectoryExists(*DirectoryPath))
+	{
+		if (!CreateDirectoryRecursively(DirectoryPath))
+		{
+			return false;
+		}
+
+		GFileManager.MakeDirectory(*DirectoryPath);
+		if (!GFileManager.DirectoryExists(*DirectoryPath))
+		{
+			// Could not make the specified directory
+			return false;
+		}
+	}
+
+	return true;
+}
+
+bool SaveStringToNewFile(const FString& FilePath, FString& SaveText)
+{
+	if (!CreateDirectoryRecursively(FilePath))
+	{
+		return false;
+	}
+
+	return FFileHelper::SaveStringToFile(*SaveText, *FilePath);
+}
+
+FString FAutomationTelemetry::ToColumnValue(uint8 Index, const FAutomationWorkerTelemetryData& Data, const FAutomationWorkerTelemetryItem& Item)
+{
+	switch (Index)
+	{
+	case Columns::Configuration:
+		return Data.Configuration;
+
+	case Columns::Platform:
+		return Data.Platform;
+
+	case Columns::DateTime:
+		return FDateTime::Now().ToString();
+
+	case Columns::TestName:
+		return Data.TestName;
+
+	case Columns::Context:
+		return Item.Context;
+
+	case Columns::DataPoint:
+		return Item.DataPoint;
+
+	case Columns::Measurement:
+		return FString::SanitizeFloat(Item.Measurement);
+
+	default:
+		return TEXT("NULL");
+	}
+}
+
+void FAutomationTelemetry::Initialize()
+{
+
+	const UAutomationControllerSettings* Settings = GetDefault<UAutomationControllerSettings>();
+	bResetTelemetryStorageOnNewSession = Settings->bResetTelemetryStorageOnNewSession;
+
+	FString TelemetryDirectoryFromSettings;
+	FParse::Value(FCommandLine::Get(), TEXT("TelemetryDirectory="), TelemetryDirectoryFromSettings, false);
+	if (TelemetryDirectoryFromSettings.IsEmpty())
+	{
+		TelemetryDirectoryFromSettings = Settings->TelemetryDirectory;
+	}
+
+	if (TelemetryDirectoryFromSettings.IsEmpty())
+	{
+		TelemetryDirectory = FPaths::AutomationDir() + TEXT("Telemetry");
+	}
+	else
+	{
+		if (FPaths::IsRelative(TelemetryDirectoryFromSettings))
+		{
+			TelemetryDirectory = FPaths::AutomationDir() + TelemetryDirectoryFromSettings;
+		}
+		else
+		{
+			TelemetryDirectory = TelemetryDirectoryFromSettings;
+		}
+	}
+
+	bIsInitialized = true;
+}
+
+bool FAutomationTelemetry::IsInitialized()
+{
+	return bIsInitialized;
+}
+
+FString FAutomationTelemetry::GetStorageFilePath(const FString& StorageName)
+{
+	return FPaths::SetExtension(TelemetryDirectory / (StorageName.IsEmpty()? TEXT("general") : StorageName), "csv");
+}
+
+bool FAutomationTelemetry::InitiateStorage(const FString& StorageFilePath)
+{
+	if (bResetTelemetryStorageOnNewSession || !FPaths::FileExists(StorageFilePath))
+	{
+		FString Header;
+		for (uint8 i = 0; i < Columns::Count; i++)
+		{
+			Header += ToColumnName(i);
+			if (i < Columns::Count - 1)
+			{
+				Header += TEXT(",");
+			}
+		}
+		Header += LINE_TERMINATOR;
+
+		if (!SaveStringToNewFile(StorageFilePath, Header))
+		{
+			return false;
+		}
+	}
+
+	return true;
+}
+
+void FAutomationTelemetry::HandleAddTelemetry(const FAutomationWorkerTelemetryData& Data)
+{
+	if (!IsInitialized())
+	{
+		Initialize();
+	}
+
+	UE_LOG(LogAutomationTelemetry, Verbose, TEXT("Received request for %s"), *Data.TestName);
+
+	FString StorageFilePath = GetStorageFilePath(Data.Storage);
+	if (!InitiateStorage(StorageFilePath))
+	{
+		UE_LOG(LogAutomationTelemetry, Error, TEXT("Could not create Telemetry file '%s'"), *StorageFilePath);
+		return;
+	}
+
+	FString Row;
+	for (const FAutomationWorkerTelemetryItem& Item : Data.Items)
+	{
+		for (uint8 i = 0; i < Columns::Count; i++)
+		{
+			Row += ToColumnValue(i, Data, Item);
+			if (i < Columns::Count - 1)
+			{
+				Row += TEXT(",");
+			}
+		}
+		Row += LINE_TERMINATOR;
+	}
+
+	if (!FFileHelper::SaveStringToFile(*Row, *StorageFilePath, FFileHelper::EEncodingOptions::AutoDetect, &IFileManager::Get(), FILEWRITE_Append))
+	{
+		UE_LOG(LogAutomationTelemetry, Error, TEXT("Could not write in Telemetry file '%s'"), *StorageFilePath);
+	}
+}
diff --git a/Engine/Source/Developer/AutomationController/Private/AutomationTelemetry.h b/Engine/Source/Developer/AutomationController/Private/AutomationTelemetry.h
new file mode 100644
--- /dev/null
+++ b/Engine/Source/Developer/AutomationController/Private/AutomationTelemetry.h
@@ -0,0 +1,84 @@
+// Copyright Epic Games, Inc. All Rights Reserved.
+
+#pragma once
+
+#include "CoreMinimal.h"
+#include "Containers/Array.h"
+#include "Containers/UnrealString.h"
+#include "AutomationWorkerMessages.h"
+
+
+class FAutomationTelemetry
+	: FNoncopyable
+{
+public:
+
+	/** Called to initialize the singleton. */
+	static void Initialize();
+	
+	/** Helper to check if the Analytics provider is Initialized */
+	static bool IsInitialized();
+
+public:
+
+	/** Handle adding telemetry data to output */
+	static void HandleAddTelemetry(const FAutomationWorkerTelemetryData& Data);
+
+private:
+	
+	static bool bIsInitialized;
+	static FString TelemetryDirectory;
+	static bool bResetTelemetryStorageOnNewSession;
+
+	enum Columns : uint8
+	{
+		Configuration,
+		Platform,
+		DateTime,
+		TestName,
+		Context,
+		DataPoint,
+		Measurement,
+
+		Count // number of columns
+	};
+
+	static FString ToColumnName(uint8 Index)
+	{
+		switch (Index)
+		{
+		case Columns::Configuration:
+			return TEXT("Configuration");
+
+		case Columns::Platform:
+			return TEXT("Platform");
+
+		case Columns::DateTime:
+			return TEXT("DateTime");
+
+		case Columns::TestName:
+			return TEXT("TestName");
+
+		case Columns::Context:
+			return TEXT("Context");
+
+		case Columns::DataPoint:
+			return TEXT("DataPoint");
+
+		case Columns::Measurement:
+			return TEXT("Measurement");
+
+		default:
+			return TEXT("Unknown");
+		}
+	}
+
+	static FString ToColumnValue(uint8 Index, const FAutomationWorkerTelemetryData& Data, const FAutomationWorkerTelemetryItem& Item);
+
+	/** Initiate telemetry storage csv file */
+	static bool InitiateStorage(const FString& StorageName);
+
+	/** Get telemetry storage file path using automation settings */
+	static FString GetStorageFilePath(const FString& StorageName);
+
+};
diff --git a/Engine/Source/Developer/AutomationController/Public/AutomationControllerSettings.h b/Engine/Source/Developer/AutomationController/Public/AutomationControllerSettings.h
--- a/Engine/Source/Developer/AutomationController/Public/AutomationControllerSettings.h
+++ b/Engine/Source/Developer/AutomationController/Public/AutomationControllerSettings.h
@@ -90,4 +90,12 @@ public:
 	/** The maximum response wait time for detecting a lost game instance (default=300sec)*/
 	UPROPERTY(Config)
 	float GameInstanceLostTimerSeconds;
+
+	/** Path to where telemetry files are saved (default=<project>/Saved/Automation/Telemetry/)*/
+	UPROPERTY(Config)
+	FString TelemetryDirectory;
+
+	/** Whether to reset data stored in telemetry file (default=false) */
+	UPROPERTY(Config)
+	bool bResetTelemetryStorageOnNewSession;
 };
diff --git a/Engine/Source/Runtime/AutomationMessages/Public/AutomationWorkerMessages.h b/Engine/Source/Runtime/AutomationMessages/Public/AutomationWorkerMessages.h
--- a/Engine/Source/Runtime/AutomationMessages/Public/AutomationWorkerMessages.h
+++ b/Engine/Source/Runtime/AutomationMessages/Public/AutomationWorkerMessages.h
@@ -820,3 +820,69 @@ struct FAutomationWorkerPerformanceDataResponse
 	UPROPERTY(EditAnywhere, Category="Message")
 	FString ErrorMessage;
 };
+
+/**
+ * Implements a message that contains telemetry data point.
+ */
+USTRUCT()
+struct FAutomationWorkerTelemetryItem
+{
+	GENERATED_USTRUCT_BODY()
+
+	/**  */
+	UPROPERTY(EditAnywhere, Category = "Message")
+	FString DataPoint;
+
+	/**  */
+	UPROPERTY(EditAnywhere, Category = "Message")
+	double Measurement;
+
+	/**  */
+	UPROPERTY(EditAnywhere, Category = "Message")
+	FString Context;
+
+	FAutomationWorkerTelemetryItem() : Measurement(0.0) {}
+
+	FAutomationWorkerTelemetryItem(FString& InDataPoint, double InMeasurement, FString& InContext)
+		: DataPoint(InDataPoint)
+		, Measurement(InMeasurement)
+		, Context(InContext)
+	{
+	}
+
+	FAutomationWorkerTelemetryItem(const FAutomationTelemetryData& InItem)
+		: DataPoint(InItem.DataPoint)
+		, Measurement(InItem.Measurement)
+		, Context(InItem.Context)
+	{
+	}
+};
+
+/**
+ * Implements a message that contains telemetry data.
+ */
+USTRUCT()
+struct FAutomationWorkerTelemetryData
+{
+	GENERATED_USTRUCT_BODY()
+
+	/**  */
+	UPROPERTY(EditAnywhere, Category = "Message")
+	FString Storage;
+
+	/**  */
+	UPROPERTY(EditAnywhere, Category = "Message")
+	FString Configuration;
+
+	/**  */
+	UPROPERTY(EditAnywhere, Category = "Message")
+	FString Platform;
+
+	/**  */
+	UPROPERTY(EditAnywhere, Category = "Message")
+	FString TestName;
+
+	/**  */
+	UPROPERTY(EditAnywhere, Category = "Message")
+	TArray<FAutomationWorkerTelemetryItem> Items;
+};
diff --git a/Engine/Source/Runtime/AutomationWorker/Private/AutomationWorkerModule.cpp b/Engine/Source/Runtime/AutomationWorker/Private/AutomationWorkerModule.cpp
--- a/Engine/Source/Runtime/AutomationWorker/Private/AutomationWorkerModule.cpp
+++ b/Engine/Source/Runtime/AutomationWorker/Private/AutomationWorkerModule.cpp
@@ -219,6 +219,8 @@ void FAutomationWorkerModule::ReportTestComplete()
 		//see if there are any more network commands left to execute
 		bool bAllLatentCommandsComplete = FAutomationTestFramework::Get().ExecuteLatentCommands();
 
+		FString TestFullName = FAutomationTestFramework::Get().GetCurrentTest()->GetTestFullName();
+
 		//structure to track error/warning/log messages
 		FAutomationTestExecutionInfo ExecutionInfo;
 
@@ -252,6 +254,11 @@ void FAutomationWorkerModule::ReportTestComplete()
 				SendAnalyticsEvents(ExecutionInfo.AnalyticsItems);
 			}
 
+			if (ExecutionInfo.TelemetryItems.Num() > 0)
+			{
+				HandleTelemetryData(ExecutionInfo.TelemetryStorage, TestFullName, ExecutionInfo.TelemetryItems);
+			}
+
 			MessageEndpoint->Send(Message, TestRequesterAddress);
 		}
 
@@ -560,6 +567,24 @@ void FAutomationWorkerModule::SendAnalyticsEvents(TArray<FString>& InAnalyticsIt
 	}
 }
 
+void FAutomationWorkerModule::HandleTelemetryData(const FString& StorageName, const FString& InTestName, const TArray<FAutomationTelemetryData>& InItems)
+{
+	FAutomationWorkerTelemetryData* Message = new FAutomationWorkerTelemetryData();
+
+	Message->Storage = StorageName;
+	Message->Platform = FPlatformProperties::PlatformName();
+	Message->Configuration = LexToString(FApp::GetBuildConfiguration());
+	Message->TestName = InTestName;
+	for (const FAutomationTelemetryData& Item : InItems)
+	{
+		Message->Items.Add(FAutomationWorkerTelemetryItem(Item));
+	}
+
+	UE_LOG(LogAutomationWorker, Log, TEXT("Sending Telemetry Data for %s"), *Message->TestName);
+
+	MessageEndpoint->Send(Message, TestRequesterAddress);
+}
+
 void FAutomationWorkerModule::RecordPerformanceAnalytics( const FAutomationPerformanceSnapshot& PerfSnapshot )
 {
 	// @todo: Pass in additional performance capture data from incoming FAutomationPerformanceSnapshot!
diff --git a/Engine/Source/Runtime/AutomationWorker/Private/AutomationWorkerModule.h b/Engine/Source/Runtime/AutomationWorker/Private/AutomationWorkerModule.h
--- a/Engine/Source/Runtime/AutomationWorker/Private/AutomationWorkerModule.h
+++ b/Engine/Source/Runtime/AutomationWorker/Private/AutomationWorkerModule.h
@@ -106,6 +106,9 @@ private:
 	/** Handles FAutomationWorkerPerformanceDataResponse messages. */
 	void HandlePerformanceDataRetrieved(const FAutomationWorkerPerformanceDataResponse& Message, const TSharedRef<IMessageContext, ESPMode::ThreadSafe>& Context);
 
+	/** Dispatches telemetry data to data collector */
+	void HandleTelemetryData(const FString& StorageName, const FString& InTestName, const TArray<FAutomationTelemetryData>& InItems);
+
 	/** Handles FAutomationTestFramework PreTestingEvents. */
 	void HandlePreTestingEvent();
 
diff --git a/Engine/Source/Runtime/Core/Private/Misc/AutomationTest.cpp b/Engine/Source/Runtime/Core/Private/Misc/AutomationTest.cpp
--- a/Engine/Source/Runtime/Core/Private/Misc/AutomationTest.cpp
+++ b/Engine/Source/Runtime/Core/Private/Misc/AutomationTest.cpp
@@ -103,7 +103,7 @@ void FAutomationTestFramework::FAutomationTestOutputDevice::Serialize( const TCH
 			{
 				if (!LoggedFailureCause.Contains(CurTest))
 				{
-					CurTest->AddError(FString::Printf(TEXT("%s %s will be marked as failing due to errors being logged"), *CurTest->GetBeautifiedTestName(), *CurTest->GetTestContext()), STACK_OFFSET);
+					CurTest->AddError(FString::Printf(TEXT("%s will be marked as failing due to errors being logged"), *CurTest->GetTestFullName()), STACK_OFFSET);
 					LoggedFailureCause.Add(CurTest);
 				}
 				CurTest->AddError(FString(V), STACK_OFFSET);
@@ -1005,6 +1005,8 @@ void FAutomationTestExecutionInfo::Clear()
 
 	Entries.Empty();
 	AnalyticsItems.Empty();
+	TelemetryItems.Empty();
+	TelemetryStorage.Empty();
 
 	Errors = 0;
 	Warnings = 0;
@@ -1177,6 +1179,24 @@ void FAutomationTestBase::AddAnalyticsItem(const FString& InAnalyticsItem)
 	ExecutionInfo.AnalyticsItems.Add(InAnalyticsItem);
 }
 
+void FAutomationTestBase::AddTelemetryData(const FString& DataPoint, double Measurement, const FString& Context)
+{
+	ExecutionInfo.TelemetryItems.Add(FAutomationTelemetryData(DataPoint, Measurement, Context));
+}
+
+void FAutomationTestBase::AddTelemetryData(const TMap <FString, double>& ValuePairs, const FString& Context)
+{
+	for (const TPair<FString, double>& Item : ValuePairs)
+	{
+		ExecutionInfo.TelemetryItems.Add(FAutomationTelemetryData(Item.Key, Item.Value, Context));
+	}
+}
+
+void FAutomationTestBase::SetTelemetryStorage(const FString& StorageName)
+{
+	ExecutionInfo.TelemetryStorage = StorageName;
+}
+
 void FAutomationTestBase::AddEvent(const FAutomationEvent& InEvent, int32 StackOffset)
 {
 	ExecutionInfo.AddEvent(InEvent, StackOffset + 1);
@@ -1287,7 +1307,7 @@ void FAutomationTestBase::GenerateTestNames(TArray<FAutomationTestInfo>& TestInf
 
 		if (ParameterNames[ParameterIndex].Len())
 		{
-			CompleteBeautifiedNames = FString::Printf(TEXT("%s.%s"), *BeautifiedTestName, *BeautifiedNames[ParameterIndex]);;
+			CompleteBeautifiedNames = FString::Printf(TEXT("%s.%s"), *BeautifiedTestName, *BeautifiedNames[ParameterIndex]);
 			CompleteTestName = FString::Printf(TEXT("%s %s"), *TestName, *ParameterNames[ParameterIndex]);
 		}
 
diff --git a/Engine/Source/Runtime/Core/Public/Misc/AutomationTest.h b/Engine/Source/Runtime/Core/Public/Misc/AutomationTest.h
--- a/Engine/Source/Runtime/Core/Public/Misc/AutomationTest.h
+++ b/Engine/Source/Runtime/Core/Public/Misc/AutomationTest.h
@@ -117,6 +117,20 @@ namespace EAutomationExpectedErrorFlags
 	}
 }
 
+struct FAutomationTelemetryData
+{
+	FString DataPoint;
+	double Measurement;
+	FString Context;
+
+	FAutomationTelemetryData(const FString& InDataPoint, double InMeasurement, const FString& InContext)
+		:DataPoint(InDataPoint)
+		, Measurement(InMeasurement)
+		, Context(InContext)
+	{
+	}
+};
+
 /** Simple class to store the results of the execution of a automation test */
 class CORE_API FAutomationTestExecutionInfo
 {
@@ -180,6 +194,12 @@ public:
 	/** Any analytics items that occurred during execution */
 	TArray<FString> AnalyticsItems;
 
+	/** Telemetry items that occurred during execution */
+	TArray<FAutomationTelemetryData> TelemetryItems;
+
+	/** Telemetry storage name set by the test */
+	FString TelemetryStorage;
+
 	/** Time to complete the task */
 	double Duration;
 
@@ -1179,6 +1199,14 @@ public:
 	/** Gets the parameter context of the test. */
 	FString GetTestContext() const { return TestParameterContext; }
 
+	/**
+	 * Returns the beautified test name with test context. Should return what is displayed in the Test Automation UI. See GenerateTestNames()
+	 */
+	virtual FString GetTestFullName() const {
+		if (GetTestContext().IsEmpty()) { return GetBeautifiedTestName(); }
+		return FString::Printf(TEXT("%s.%s"), *GetBeautifiedTestName(), *GetTestContext());
+	}
+
 	/**
 	 * Pure virtual method; returns the number of participants for this test
 	 *
@@ -1256,6 +1284,30 @@ public:
 	 */
 	virtual void AddAnalyticsItem(const FString& InAnalyticsItem);
 
+	/**
+	 * Adds a telemetry data point measurement
+	 *
+	 * @param	DataPoint	Name of the Data point
+	 * @param	Measurement	Value to associate to the data point
+	 * @param	Context		optional context associated with the data point
+	 */
+	virtual void AddTelemetryData(const FString& DataPoint, double Measurement, const FString& Context = TEXT(""));
+
+	/**
+	 * Adds several telemetry data point measurements
+	 *
+	 * @param	ValuePairs	value pair of Name and Measurement of several Data points
+	 * @param	Context		optional context associated with the data point
+	 */
+	virtual void AddTelemetryData(const TMap<FString, double>& ValuePairs, const FString& Context = TEXT(""));
+
+	/**
+	 * Set telemetry storage name
+	 *
+	 * @param	StorageName	Name of the data storage
+	 */
+	virtual void SetTelemetryStorage(const FString& StorageName);
+
 	/**
 	 * Returns whether this test has any errors associated with it or not
 	 *
diff --git a/Engine/Source/Runtime/Engine/Private/Tests/AutomationCommon.cpp b/Engine/Source/Runtime/Engine/Private/Tests/AutomationCommon.cpp
--- a/Engine/Source/Runtime/Engine/Private/Tests/AutomationCommon.cpp
+++ b/Engine/Source/Runtime/Engine/Private/Tests/AutomationCommon.cpp
@@ -202,7 +202,7 @@ namespace AutomationCommon
 			if (FAutomationTestBase* CurrentTest = FAutomationTestFramework::Get().GetCurrentTest())
 			{
 				Context = CurrentTest->GetTestContext();
-				if (Context.IsEmpty()) { Context = CurrentTest->GetTestName(); }
+				if (Context.IsEmpty()) { Context = CurrentTest->GetTestFullName(); }
 			}
 
 			FAutomationScreenshotData Data = BuildScreenshotData(Context, TEXT(""), ImageName, InWidth, InHeight);
