Fix missing include in AutomationTelemetry.cpp

#rnx

diff --git a/Engine/Source/Developer/AutomationController/Private/AutomationTelemetry.cpp b/Engine/Source/Developer/AutomationController/Private/AutomationTelemetry.cpp
--- a/Engine/Source/Developer/AutomationController/Private/AutomationTelemetry.cpp
+++ b/Engine/Source/Developer/AutomationController/Private/AutomationTelemetry.cpp
@@ -6,6 +6,7 @@
 #include "Misc/FileHelper.h"
 #include "Misc/Paths.h"
 #include "Misc/DateTime.h"
+#include "Misc/CommandLine.h"
 #include "AutomationControllerSettings.h"
 
 DEFINE_LOG_CATEGORY_STATIC(LogAutomationTelemetry, Log, All)
