Close the PIE window after automated test pass has finished

#rnx

diff --git a/Engine/Source/Developer/AutomationController/Private/AutomationControllerManger.cpp b/Engine/Source/Developer/AutomationController/Private/AutomationControllerManger.cpp
--- a/Engine/Source/Developer/AutomationController/Private/AutomationControllerManger.cpp
+++ b/Engine/Source/Developer/AutomationController/Private/AutomationControllerManger.cpp
@@ -23,6 +23,7 @@
 #include "Misc/EngineVersion.h"
 #include "Misc/FileHelper.h"
 #include "PlatformHttp.h"
+#include "UnrealEd.h"
 
 #include "AutomationTelemetry.h"
 
@@ -371,6 +372,12 @@ void FAutomationControllerManager::ProcessAvailableTasks()
 				if ( ClusterDistributionMask == 0 )
 				{
 					ProcessResults();
+
+					// Close play window
+					#if WITH_EDITOR
+						GUnrealEd->RequestEndPlayMap();
+					#endif
+
 					//Notify the graphical layout we are done processing results.
 					TestsCompleteDelegate.Broadcast();
 				}
