Fix AutomationControllerManger.cpp warnings: Inconsistent preprocessor state

#rnx

diff --git a/Engine/Source/Developer/AutomationController/Private/AutomationControllerManger.cpp b/Engine/Source/Developer/AutomationController/Private/AutomationControllerManger.cpp
--- a/Engine/Source/Developer/AutomationController/Private/AutomationControllerManger.cpp
+++ b/Engine/Source/Developer/AutomationController/Private/AutomationControllerManger.cpp
@@ -27,7 +27,8 @@
 #include "AutomationTelemetry.h"
 
 #if WITH_EDITOR
-#include "UnrealEd.h"
+#include "UnrealEdGlobals.h"
+#include "Editor/UnrealEdEngine.h"
 #include "Logging/MessageLog.h"
 #endif
 #include "Async/ParallelFor.h"
