Improve test exclusion list mechanic

* Rename blacklist to excludelist (requested by high management)
* Support section exclusion rule to be able to exclude entire section of tests (ie PathTracing is only supported on Win64 for now)
* Mark excluded test as skipped in the report instead of entirely removed for test list. Check for exclusion just before running the test.
* Remove NotEnoughParticipant state in favor of Skipped (same conditions lead to Skipped state with appropriate messaging)
* Add support for exclusion management from the Test Automation window. (added a column at the end of each row)
* Expose device information to UE test report
* Add support for metadata in Gauntlet test report for Horde

Limitations:
* Management through the UI is limited to which test is available through in the active worker node. That's mean Runtime only tests are not listed from a worker that is Editor(the default) and platform specific are not clearly identified.
* For platforms, the mechanic to access their config and save it will remain to be done. In the meantime, it needs to be done manually through the target platform config file.

diff --git a/Engine/Source/Developer/AutomationController/AutomationController.Build.cs b/Engine/Source/Developer/AutomationController/AutomationController.Build.cs
--- a/Engine/Source/Developer/AutomationController/AutomationController.Build.cs
+++ b/Engine/Source/Developer/AutomationController/AutomationController.Build.cs
@@ -10,7 +10,8 @@ namespace UnrealBuildTool.Rules
 				new string[]
 				{
 					"Core",
-					"CoreUObject"
+					"CoreUObject",
+					"AutomationTest",
 				}
 			); 
 			
@@ -18,7 +19,7 @@ namespace UnrealBuildTool.Rules
 				new string[]
 				{
                     "AssetRegistry",
-                    "AutomationMessages",
+					"AutomationMessages",
 					"UnrealEdMessages",
                     "MessageLog",
                     "Json",
@@ -40,16 +41,17 @@ namespace UnrealBuildTool.Rules
                 );
             }
 
-            PrivateIncludePathModuleNames.AddRange(
+			PublicIncludePathModuleNames.AddRange(
 				new string[]
 				{
-					"MessagingCommon",
+					"AutomationTest"
 				}
 			);
 
-			PrivateIncludePaths.AddRange(
+			PrivateIncludePathModuleNames.AddRange(
 				new string[]
 				{
+					"MessagingCommon"
 				}
 			);
 
diff --git a/Engine/Source/Developer/AutomationController/Private/AutomationControllerManager.h b/Engine/Source/Developer/AutomationController/Private/AutomationControllerManager.h
--- a/Engine/Source/Developer/AutomationController/Private/AutomationControllerManager.h
+++ b/Engine/Source/Developer/AutomationController/Private/AutomationControllerManager.h
@@ -36,6 +36,9 @@ public:
 	UPROPERTY()
 	EAutomationState State;
 
+	UPROPERTY()
+	FString DeviceInstance;
+
 	FAutomatedTestResult()
 	{
 		Warnings = 0;
@@ -100,8 +103,7 @@ struct FAutomatedTestPassResults
 
 public:
 	FAutomatedTestPassResults()
-		: ClientDescriptor()
-		, ReportCreatedOn(0)
+		: ReportCreatedOn(0)
 		, Succeeded(0)
 		, SucceededWithWarnings(0)
 		, Failed(0)
@@ -111,24 +113,10 @@ public:
 		, ComparisonExported(false)
 		, IsRequired(false)
 	{
-		if (FEngineVersion::Current().HasChangelist())
-		{
-			ClientDescriptor = FEngineVersion::Current().GetBranch()
-				+ TEXT(" - ")
-				+ FString::FromInt(FEngineVersion::Current().GetChangelist())
-				+ TEXT(" - ");
-		}
-
-		if (FPlatformProperties::RequiresCookedData())
-		{
-			ClientDescriptor += TEXT("Cooked ");
-		}
-
-		ClientDescriptor += FPlatformProperties::IniPlatformName();
 	}
 
 	UPROPERTY()
-	FString ClientDescriptor;
+	TArray<FAutomationDeviceInfo> Devices;
 
 	UPROPERTY()
 	FDateTime ReportCreatedOn;
@@ -181,6 +169,7 @@ public:
 
 	void ClearAllEntries()
 	{
+		Devices.Empty();
 		Succeeded = 0;
 		SucceededWithWarnings = 0;
 		Failed = 0;
diff --git a/Engine/Source/Developer/AutomationController/Private/AutomationControllerManger.cpp b/Engine/Source/Developer/AutomationController/Private/AutomationControllerManger.cpp
--- a/Engine/Source/Developer/AutomationController/Private/AutomationControllerManger.cpp
+++ b/Engine/Source/Developer/AutomationController/Private/AutomationControllerManger.cpp
@@ -34,9 +34,8 @@
 #include "Async/ParallelFor.h"
 
 // these strings are parsed by Gauntlet (AutomationLogParser) so make sure changes are replicated there!
-#define AutomationTestStarting		TEXT("Test Started. Name={%s}")
-#define AutomationSuccessFormat		TEXT("Test Completed. Result={Passed} Name={%s} Path={%s}")
-#define AutomationFailureFormat		TEXT("Test Completed. Result={Failed} Name={%s} Path={%s}")
+#define AutomationTestStarting		TEXT("Test Started. Name={%s} Path={%s}")
+#define AutomationStateFormat		TEXT("Test Completed. Result={%s} Name={%s} Path={%s}")
 
 #define BeginEventsFormat			TEXT("BeginEvents: %s")
 #define EndEventsFormat				TEXT("EndEvents: %s")
@@ -102,9 +101,21 @@ bool FAutomatedTestPassResults::ReflectResultStateToReport(IAutomationReportPtr&
 void FAutomatedTestPassResults::UpdateTestResultStatus(const IAutomationReportPtr& TestReport, EAutomationState State, bool bHasWarning)
 {
 	FAutomatedTestResult& TestResult = GetTestResult(TestReport);
-	TestResult.State = State;
+	if (TestResult.State == State)
+		return;
 
 	// Book keeping
+	// from transient states
+	switch (TestResult.State)
+	{
+	case EAutomationState::InProcess:
+		InProcess--;
+		break;
+	case EAutomationState::NotRun:
+		NotRun--;
+		break;
+	}
+	// to new state
 	switch (State)
 	{
 	case EAutomationState::Success:
@@ -116,25 +127,22 @@ void FAutomatedTestPassResults::UpdateTestResultStatus(const IAutomationReportPt
 		{
 			Succeeded++;
 		}
-		InProcess--;
 		break;
 	case EAutomationState::Fail:
 		Failed++;
-		InProcess--;
 		break;
 	case EAutomationState::InProcess:
-		NotRun--;
 		InProcess++;
 		break;
 	case EAutomationState::NotRun:
 		NotRun++;
-		InProcess--;
 		break;
-	case EAutomationState::NotEnoughParticipants:
-		// Those never get to be set as InProcess
-		NotRun--;
+	case EAutomationState::Skipped:
+		// Those are not counted
 		break;
 	}
+	// commit the change
+	TestResult.State = State;
 }
 
 FAutomationControllerManager::FAutomationControllerManager()
@@ -309,6 +317,13 @@ void FAutomationControllerManager::RunTests(const bool bInIsLocalSession)
 			UE_LOG(LogAutomationController, Log, TEXT("Sending Reset Tests to %s"), *MessageAddress.ToString());
 
 			MessageEndpoint->Send(FMessageEndpoint::MakeMessage<FAutomationWorkerResetTests>(), MessageAddress);
+
+			// Store devices info into the json report.
+			if (JsonTestPassResults.IsRequired)
+			{
+				const FAutomationDeviceInfo& DeviceInfo = DeviceClusterManager.GetDeviceInfo(ClusterIndex, DeviceIndex);
+				JsonTestPassResults.Devices.Add(DeviceInfo);
+			}
 		}
 	}
 	
@@ -503,7 +518,7 @@ void FAutomationControllerManager::ReportTestResults()
 	UE_LOG(LogAutomationController, Log, TEXT("Test Pass Results:"));
 	for ( int32 i = 0; i < JsonTestPassResults.Tests.Num(); i++ )
 	{
-		UE_LOG(LogAutomationController, Log, TEXT("%s: %s"), *JsonTestPassResults.Tests[i].TestDisplayName, ToString(JsonTestPassResults.Tests[i].State));
+		UE_LOG(LogAutomationController, Log, TEXT("%s: %s"), *JsonTestPassResults.Tests[i].TestDisplayName, *AutomationStateToString(JsonTestPassResults.Tests[i].State));
 	}
 }
 
@@ -515,6 +530,7 @@ void FAutomationControllerManager::CollectTestResults(TSharedPtr<IAutomationRepo
 		FAutomatedTestResult& TestResult = JsonTestPassResults.GetTestResult(Report);
 		TestResult.SetEvents(Results.GetEntries(), Results.GetWarningTotal(), Results.GetErrorTotal());
 		TestResult.SetArtifacts(Results.Artifacts);
+		TestResult.DeviceInstance = Results.GameInstance;
 
 		JsonTestPassResults.TotalDuration += Results.Duration;
 
@@ -700,7 +716,7 @@ void FAutomationControllerManager::ExecuteNextTask( int32 ClusterIndex, OUT bool
 			{
 				// Get the status of the test
 				EAutomationState TestState = NextTest->GetState(ClusterIndex, CurrentTestPass);
-				if ( TestState == EAutomationState::NotRun )
+				if (TestState == EAutomationState::NotRun)
 				{
 					// Reserve this device for the test
 					DeviceClusterManager.SetTest(ClusterIndex, DeviceIndex, NextTest);
@@ -708,19 +724,20 @@ void FAutomationControllerManager::ExecuteNextTask( int32 ClusterIndex, OUT bool
 
 					// If we now have enough devices reserved for the test, run it!
 					TArray<FMessageAddress> DeviceAddresses = DeviceClusterManager.GetDevicesReservedForTest(ClusterIndex, NextTest);
-					if ( DeviceAddresses.Num() == NextTest->GetNumParticipantsRequired() )
+					if (DeviceAddresses.Num() == NextTest->GetNumParticipantsRequired())
 					{
 						// Send it to each device
-						for ( int32 AddressIndex = 0; AddressIndex < DeviceAddresses.Num(); ++AddressIndex )
+						for (int32 AddressIndex = 0; AddressIndex < DeviceAddresses.Num(); ++AddressIndex)
 						{
 							FAutomationTestResults TestResults;
+							TSharedPtr< IAutomationReport > Test = TestsRunThisPass[AddressIndex];
 
-							UE_LOG(LogAutomationController, Display, AutomationTestStarting, *TestsRunThisPass[AddressIndex]->GetDisplayName(), *TestsRunThisPass[AddressIndex]->GetCommand());
+							UE_LOG(LogAutomationController, Display, AutomationTestStarting, *Test->GetDisplayName(), *Test->GetFullTestPath());
 							TestResults.State = EAutomationState::InProcess;
 
 							if (JsonTestPassResults.IsRequired)
 							{
-								JsonTestPassResults.UpdateTestResultStatus(NextTest, EAutomationState::InProcess);
+								JsonTestPassResults.UpdateTestResultStatus(NextTest, TestResults.State);
 								// Save the whole pass report so that if the next test triggers a critical failure we are not left with no pass report.
 								GenerateJsonTestPassSummary(JsonTestPassResults);
 							}
@@ -735,7 +752,7 @@ void FAutomationControllerManager::ExecuteNextTask( int32 ClusterIndex, OUT bool
 							// Send the test to the device for execution!
 							UE_LOG(LogAutomationController, Log, TEXT("Sending RunTest %s to %s"), *NextTest->GetDisplayName(), *DeviceAddress.ToString());
 
-							MessageEndpoint->Send(FMessageEndpoint::MakeMessage<FAutomationWorkerRunTests>(ExecutionCount, AddressIndex, NextTest->GetCommand(), NextTest->GetDisplayName(), bSendAnalytics), DeviceAddress);
+							MessageEndpoint->Send(FMessageEndpoint::MakeMessage<FAutomationWorkerRunTests>(ExecutionCount, AddressIndex, NextTest->GetCommand(), NextTest->GetDisplayName(), NextTest->GetFullTestPath(), bSendAnalytics), DeviceAddress);
 
 							// Add a test so we can check later if the device is still active
 							TestRunningArray.Add(FTestRunningInfo(DeviceAddress));
@@ -761,9 +778,9 @@ void FAutomationControllerManager::ExecuteNextTask( int32 ClusterIndex, OUT bool
 			if ( GetNumDevicesInCluster(ClusterIndex) < CurrentTest->GetNumParticipantsRequired() )
 			{
 				FAutomationTestResults TestResults;
-				TestResults.State = EAutomationState::NotEnoughParticipants;
+				TestResults.State = EAutomationState::Skipped;
 				TestResults.GameInstance = DeviceClusterManager.GetClusterDeviceName(ClusterIndex, 0);
-				TestResults.AddEvent(FAutomationEvent(EAutomationEventType::Warning, FString::Printf(TEXT("Needed %d devices to participate, Only had %d available."), CurrentTest->GetNumParticipantsRequired(), DeviceClusterManager.GetNumDevicesInCluster(ClusterIndex))));
+				TestResults.AddEvent(FAutomationEvent(EAutomationEventType::Warning, FString::Printf(TEXT("Skipped because the test needed %d devices to participate, Only had %d available."), CurrentTest->GetNumParticipantsRequired(), DeviceClusterManager.GetNumDevicesInCluster(ClusterIndex))));
 
 				CurrentTest->SetResults(ClusterIndex, CurrentTestPass, TestResults);
 				CollectTestResults(CurrentTest, TestResults);
@@ -907,6 +924,18 @@ void FAutomationControllerManager::ProcessResults()
 					return false;
 				}
 
+				if (A.State == EAutomationState::Skipped)
+				{
+					if (B.State == EAutomationState::Skipped)
+						return  (A.FullTestPath < B.FullTestPath);
+					else
+						return true;
+				}
+				else if (B.State == EAutomationState::Skipped)
+				{
+					return false;
+				}
+
 				if (A.GetWarningTotal() > 0)
 				{
 					if (B.GetWarningTotal() > 0)
@@ -1355,18 +1384,23 @@ void FAutomationControllerManager::ReportAutomationResult(const TSharedPtr<IAuto
 #endif
 
 	const FAutomationTestResults& Results = InReport->GetResults(ClusterIndex, PassIndex);
+	const FString StateString = AutomationStateToString(Results.State);
 
 	// write results to editor panel
 #if WITH_EDITOR
 	FFormatNamedArguments Arguments;
 	Arguments.Add(TEXT("TestName"), FText::FromString(InReport->GetFullTestPath()));
-	Arguments.Add(TEXT("Result"), Results.State != EAutomationState::Success ? LOCTEXT("Failed", "Failed") : LOCTEXT("Passed", "Passed"));
+	Arguments.Add(TEXT("Result"), FText::FromString(StateString));
 	TSharedRef<FTextToken> Token = FTextToken::Create(FText::Format(LOCTEXT("TestSuccessOrFailure", "Test '{TestName}' completed with result '{Result}'"), Arguments));
 	
 	if (Results.State == EAutomationState::Success)
 	{
 		AutomationEditorLog.Info()->AddToken(Token);
 	}
+	else if (Results.State == EAutomationState::Skipped)
+	{
+		AutomationEditorLog.Warning()->AddToken(Token);
+	}
 	else
 	{
 		AutomationEditorLog.Error()->AddToken(Token);
@@ -1374,15 +1408,14 @@ void FAutomationControllerManager::ReportAutomationResult(const TSharedPtr<IAuto
 #endif
 
 	// Now log
-	if (Results.State == EAutomationState::Success)
+	FString ResultString = FString::Printf(AutomationStateFormat, *StateString, *InReport->GetDisplayName(), *InReport->GetFullTestPath());
+	if (Results.State == EAutomationState::Fail)
 	{
-		FString SuccessString = FString::Printf(AutomationSuccessFormat, *InReport->GetDisplayName(), *InReport->GetFullTestPath());
-		UE_LOG(LogAutomationController, Display, TEXT("%s"), *SuccessString);
+		UE_LOG(LogAutomationController, Error, TEXT("%s"), *ResultString);
 	}
 	else
 	{
-		FString FailureString = FString::Printf(AutomationFailureFormat, *InReport->GetDisplayName(), *InReport->GetFullTestPath());
-		UE_LOG(LogAutomationController, Error, TEXT("%s"), *FailureString);
+		UE_LOG(LogAutomationController, Display, TEXT("%s"), *ResultString);
 	}
 
 	// bracket these for easy parsing
@@ -1429,7 +1462,7 @@ void FAutomationControllerManager::HandleRunTestsReplyMessage(const FAutomationW
 	{
 		FAutomationTestResults TestResults;
 
-		TestResults.State = Message.Success ? EAutomationState::Success : EAutomationState::Fail;
+		TestResults.State = Message.State;
 		TestResults.Duration = Message.Duration;
 
 		// Mark device as back on the market
diff --git a/Engine/Source/Developer/AutomationController/Private/AutomationDeviceClusterManager.cpp b/Engine/Source/Developer/AutomationController/Private/AutomationDeviceClusterManager.cpp
--- a/Engine/Source/Developer/AutomationController/Private/AutomationDeviceClusterManager.cpp
+++ b/Engine/Source/Developer/AutomationController/Private/AutomationDeviceClusterManager.cpp
@@ -63,42 +63,42 @@ FString FAutomationDeviceClusterManager::GetGroupNameForDevice(const FDeviceStat
 
 	if( (DeviceGroupFlags & (1 << EAutomationDeviceGroupTypes::MachineName)) > 0 )
 	{
-		OutGroupName += DeviceState.DeviceName + TEXT("-");
+		OutGroupName += DeviceState.Info.DeviceName + TEXT("-");
 	}
 
 	if( (DeviceGroupFlags & (1 << EAutomationDeviceGroupTypes::Platform)) > 0 )
 	{
-		OutGroupName += DeviceState.PlatformName + TEXT("-");
+		OutGroupName += DeviceState.Info.Platform + TEXT("-");
 	}
 
 	if( (DeviceGroupFlags & (1 << EAutomationDeviceGroupTypes::OSVersion)) > 0 )
 	{
-		OutGroupName += DeviceState.OSVersionName + TEXT("-");
+		OutGroupName += DeviceState.Info.OSVersion + TEXT("-");
 	}
 
 	if( (DeviceGroupFlags & (1 << EAutomationDeviceGroupTypes::Model)) > 0 )
 	{
-		OutGroupName += DeviceState.ModelName + TEXT("-");
+		OutGroupName += DeviceState.Info.Model + TEXT("-");
 	}
 
 	if( (DeviceGroupFlags & (1 << EAutomationDeviceGroupTypes::GPU)) > 0 )
 	{
-		OutGroupName += DeviceState.GPUName + TEXT("-");
+		OutGroupName += DeviceState.Info.GPU + TEXT("-");
 	}
 
 	if( (DeviceGroupFlags & (1 << EAutomationDeviceGroupTypes::CPUModel)) > 0 )
 	{
-		OutGroupName += DeviceState.CPUModelName + TEXT("-");
+		OutGroupName += DeviceState.Info.CPUModel + TEXT("-");
 	}
 
 	if( (DeviceGroupFlags & (1 << EAutomationDeviceGroupTypes::RamInGB)) > 0 )
 	{
-		OutGroupName += FString::Printf(TEXT("%uGB Ram-"),DeviceState.RAMInGB);
+		OutGroupName += FString::Printf(TEXT("%uGB Ram-"),DeviceState.Info.RAMInGB);
 	}
 
 	if( (DeviceGroupFlags & (1 << EAutomationDeviceGroupTypes::RenderMode)) > 0 )
 	{
-		OutGroupName += DeviceState.RenderModeName + TEXT("-");
+		OutGroupName += DeviceState.Info.RenderMode + TEXT("-");
 	}
 
 	if( OutGroupName.Len() > 0 )
@@ -144,7 +144,7 @@ void FAutomationDeviceClusterManager::ReGroupDevices( const uint32 GroupFlags )
 		{
 			FDeviceCluster NewCluster;
 			NewCluster.ClusterName = GroupName;
-			NewCluster.DeviceTypeName = DeviceIt->PlatformName;
+			NewCluster.DeviceTypeName = DeviceIt->Info.Platform;
 			NewCluster.Devices.Add(*DeviceIt);
 			Clusters.Add(NewCluster);
 		}
@@ -209,9 +209,15 @@ FString FAutomationDeviceClusterManager::GetClusterDeviceName(const int32 Cluste
 {
 	check((ClusterIndex >= 0) && (ClusterIndex < Clusters.Num()));
 	check((DeviceIndex >= 0) && (DeviceIndex < Clusters[ClusterIndex].Devices.Num()));
-	return Clusters[ClusterIndex].Devices[DeviceIndex].GameInstanceName;
+	return Clusters[ClusterIndex].Devices[DeviceIndex].Info.Instance;
 }
 
+const FAutomationDeviceInfo& FAutomationDeviceClusterManager::GetDeviceInfo(const int32 ClusterIndex, const int32 DeviceIndex) const
+{
+	check((ClusterIndex >= 0) && (ClusterIndex < Clusters.Num()));
+	check((DeviceIndex >= 0) && (DeviceIndex < Clusters[ClusterIndex].Devices.Num()));
+	return Clusters[ClusterIndex].Devices[DeviceIndex].Info;
+}
 
 bool FAutomationDeviceClusterManager::FindDevice(const FMessageAddress& MessageAddress, int32& OutClusterIndex, int32& OutDeviceIndex)
 {
diff --git a/Engine/Source/Developer/AutomationController/Private/AutomationDeviceClusterManager.h b/Engine/Source/Developer/AutomationController/Private/AutomationDeviceClusterManager.h
--- a/Engine/Source/Developer/AutomationController/Private/AutomationDeviceClusterManager.h
+++ b/Engine/Source/Developer/AutomationController/Private/AutomationDeviceClusterManager.h
@@ -3,9 +3,60 @@
 #pragma once
 
 #include "CoreMinimal.h"
+#include "UObject/ObjectMacros.h"
 #include "IAutomationReport.h"
 #include "AutomationWorkerMessages.h"
 #include "IMessageContext.h"
+#include "AutomationDeviceClusterManager.generated.h"
+
+/**
+* Hold information about the Device
+*/
+USTRUCT()
+struct FAutomationDeviceInfo
+{
+	GENERATED_BODY()
+public:
+	/** The name of device */
+	UPROPERTY()
+	FString DeviceName;
+
+	/** The instance name */
+	UPROPERTY()
+	FString Instance;
+
+	/** The name of the platform */
+	UPROPERTY()
+	FString Platform;
+
+	/** The name of the operating system version */
+	UPROPERTY()
+	FString OSVersion;
+
+	/** The name of the device model */
+	UPROPERTY()
+	FString Model;
+
+	/** The name of the GPU */
+	UPROPERTY()
+	FString GPU;
+
+	/** The name of the CPU model */
+	UPROPERTY()
+	FString CPUModel;
+
+	/** The amount of RAM this device has in gigabytes */
+	UPROPERTY()
+	uint32 RAMInGB = 0;
+
+	/** The name of the current render mode */
+	UPROPERTY()
+	FString RenderMode;
+
+	/** The name of the current RHI */
+	UPROPERTY()
+	FString RHI;
+};
 
 /**
  * Managers groups of devices for job distribution
@@ -49,6 +100,9 @@ public:
 	/** Returns the name of a device within a cluster. */
 	FString GetClusterDeviceName(const int32 ClusterIndex, const int32 DeviceIndex) const;
 
+	/** Returns the name of a device within a cluster. */
+	const FAutomationDeviceInfo& GetDeviceInfo(const int32 ClusterIndex, const int32 DeviceIndex) const;
+
 	/** 
 	 * Finds the cluster/device index for a particular GUID.
 	 *
@@ -139,48 +193,25 @@ private:
 		FDeviceState(FMessageAddress NewMessageAddress, const FAutomationWorkerFindWorkersResponse& Message )
 		{
 			DeviceMessageAddress = NewMessageAddress;
-			DeviceName = Message.DeviceName;
-			PlatformName = Message.Platform;
-			OSVersionName = Message.OSVersionName;
-			ModelName = Message.ModelName;
-			GPUName = Message.GPUName;
-			CPUModelName = Message.CPUModelName;
-			RAMInGB = Message.RAMInGB;
-			RenderModeName = Message.RenderModeName;
+			Info.DeviceName = Message.DeviceName;
+			Info.Platform = Message.Platform;
+			Info.OSVersion = Message.OSVersionName;
+			Info.Model = Message.ModelName;
+			Info.GPU = Message.GPUName;
+			Info.CPUModel = Message.CPUModelName;
+			Info.RAMInGB = Message.RAMInGB;
+			Info.RenderMode = Message.RenderModeName;
+			Info.RHI = Message.RHIName;
+			Info.Instance = Message.InstanceName;
 			Report.Reset();
-			GameInstanceName = Message.InstanceName;
 			IsDeviceAvailable = true;
 		}
 
 		/** Network address for device */
 		FMessageAddress DeviceMessageAddress;
 
-		/** The instance name */
-		FString GameInstanceName;
-
-		/** The name of device */
-		FString DeviceName;
-
-		/** The name of the platform */
-		FString PlatformName;
-
-		/** The name of the operating system version */
-		FString OSVersionName;
-
-		/** The name of the device model */
-		FString ModelName;
-
-		/** The name of the GPU */
-		FString GPUName;
-
-		/** The name of the CPU model */
-		FString CPUModelName;
-
-		/** The amount of RAM this device has in gigabytes */
-		uint32 RAMInGB;
-
-		/** The name of the current render mode */
-		FString RenderModeName;
+		/** The Device full info */
+		FAutomationDeviceInfo Info;
 
 		/** NULL if this device is available to do work*/
 		TSharedPtr <IAutomationReport> Report;
diff --git a/Engine/Source/Developer/AutomationController/Private/AutomationReport.cpp b/Engine/Source/Developer/AutomationController/Private/AutomationReport.cpp
--- a/Engine/Source/Developer/AutomationController/Private/AutomationReport.cpp
+++ b/Engine/Source/Developer/AutomationController/Private/AutomationReport.cpp
@@ -1,6 +1,7 @@
 // Copyright Epic Games, Inc. All Rights Reserved.
 
 #include "AutomationReport.h"
+#include "AutomationTestExcludelist.h"
 #include "Misc/FilterCollection.h"
 #include "HAL/FileManager.h"
 #include "Misc/FileHelper.h"
@@ -19,6 +20,13 @@ FAutomationReport::FAutomationReport(FAutomationTestInfo& InTestInfo, bool InIsP
 	{
 		bEnabled = true;
 	}
+	// Get exclude test info from Config/DefaultEngine.ini
+	if (auto Entry = UAutomationTestExcludelist::Get()->GetExcludeTestEntry(TestInfo.GetFullTestPath()))
+	{
+		ExcludeTestInfo = *Entry;
+		ExcludeTestInfo.SetPropagation(TestInfo.GetFullTestPath());
+		bNeedToSkip = true;
+	}
 }
 
 void FAutomationReport::Empty()
@@ -444,7 +452,7 @@ void FAutomationReport::GetCompletionStatus(const int32 ClusterIndex, const int3
 		{
 			IsEnabled() ? OutCompletionState.NumEnabledTestsFailed++ : OutCompletionState.NumDisabledTestsFailed++;
 		}
-		else if( CurrentState == EAutomationState::NotEnoughParticipants )
+		else if( CurrentState == EAutomationState::Skipped )
 		{
 			IsEnabled() ? OutCompletionState.NumEnabledTestsCouldntBeRun++ : OutCompletionState.NumDisabledTestsCouldntBeRun++;
 		}
@@ -518,6 +526,7 @@ TSharedPtr<IAutomationReport> FAutomationReport::EnsureReportExists(FAutomationT
 {
 	//Split New Test Name by the first "." found
 	FString NameToMatch = InTestInfo.GetDisplayName();
+	FString FullPath = InTestInfo.GetFullTestPath();
 	FString NameRemainder;
 	//if this is a leaf test (no ".")
 	if (!InTestInfo.GetDisplayName().Split(TEXT("."), &NameToMatch, &NameRemainder))
@@ -529,6 +538,13 @@ TSharedPtr<IAutomationReport> FAutomationReport::EnsureReportExists(FAutomationT
 	{
 		// Set the test info name to be the remaining string
 		InTestInfo.SetDisplayName( NameRemainder );
+		// Update the fullpath
+		int32 Pos = FullPath.Find(NameToMatch + TEXT("."));
+		if (Pos >= 0)
+		{
+			FullPath.LeftChopInline(FullPath.Len() - Pos);
+		}
+		FullPath += NameToMatch;
 	}
 
 	uint32 NameToMatchHash = GetTypeHash(NameToMatch);
@@ -561,7 +577,7 @@ TSharedPtr<IAutomationReport> FAutomationReport::EnsureReportExists(FAutomationT
 		else
 		{
 			// Create a parent node
-			FAutomationTestInfo ParentTestInfo(NameToMatch, TEXT(""), TEXT(""), InTestInfo.GetTestFlags(), InTestInfo.GetNumParticipantsRequired());
+			FAutomationTestInfo ParentTestInfo(NameToMatch, FullPath, TEXT(""), InTestInfo.GetTestFlags(), InTestInfo.GetNumParticipantsRequired());
 			MatchTest = MakeShareable(new FAutomationReport(ParentTestInfo, true));
 		}
 
@@ -632,7 +648,7 @@ TSharedPtr<IAutomationReport> FAutomationReport::GetNextReportToExecute(bool& bO
 		{
 			EAutomationState TestState = GetState(ClusterIndex,PassIndex);
 			//if any enabled test hasn't been run yet or is in process
-			if ((TestState != EAutomationState::Success) && (TestState != EAutomationState::Fail) && (TestState != EAutomationState::NotEnoughParticipants))
+			if (TestState == EAutomationState::NotRun || TestState == EAutomationState::InProcess)
 			{
 				//make sure we announce we are NOT done with all tests
 				bOutAllTestsComplete = false;
@@ -812,3 +828,95 @@ void FAutomationReport::StopRunningTest()
 		ChildReports[ChildIndex]->StopRunningTest();
 	}
 }
+
+bool FAutomationReport::IsToBeSkipped(FName* OutReason, bool* OutWarn) const
+{
+	if (bNeedToSkip)
+	{
+		if (OutReason != nullptr)
+		{
+			*OutReason = ExcludeTestInfo.Reason;
+		}
+
+		if (OutWarn != nullptr)
+		{
+			*OutWarn = ExcludeTestInfo.Warn;
+		}
+
+		return true;
+	}
+
+	return false;
+}
+
+bool FAutomationReport::IsToBeSkippedByPropagation() const
+{
+	return bNeedToSkip && ExcludeTestInfo.bIsPropagated;
+}
+
+void FAutomationReport::SetSkipFlag(bool bEnableSkip, const FAutomationTestExcludelistEntry* Template, bool bFromPropagation)
+{
+	if (IsToBeSkipped() == bEnableSkip)
+	{
+		if (!bEnableSkip || Template == nullptr)
+			return;
+
+		if (!bFromPropagation)
+		{
+			// Remove previous entry in the config
+			UAutomationTestExcludelist::Get()->RemoveFromExcludeTest(TestInfo.GetFullTestPath());
+		}
+	}
+
+	if (!bFromPropagation && !ExcludeTestInfo.IsEmpty() && ExcludeTestInfo.bIsPropagated)
+		return; // Propagated exclusion can't be changed directly
+
+	bNeedToSkip = bEnableSkip;
+
+	if (Template != nullptr)
+	{
+		// Update the entry
+		ExcludeTestInfo = *Template;
+		ExcludeTestInfo.bIsPropagated = bFromPropagation;
+	}
+
+	if (bFromPropagation)
+	{
+		if (!bNeedToSkip)
+		{
+			ExcludeTestInfo.Reset();
+		}
+	}
+	else
+	{
+		auto ExcludedTestCached = UAutomationTestExcludelist::Get();
+		if (bNeedToSkip)
+		{
+			check(Template != nullptr);
+			ExcludedTestCached->AddToExcludeTest(TestInfo.GetFullTestPath(), *Template);
+		}
+		else
+		{
+			ExcludedTestCached->RemoveFromExcludeTest(TestInfo.GetFullTestPath());
+		}
+
+		ExcludedTestCached->SaveConfig();
+	}
+
+	// Propagate to children
+	if (IsParent())
+	{
+		for (IAutomationReportPtr Child : GetChildReports())
+		{
+			Child->SetSkipFlag(bNeedToSkip, Template, true);
+		}
+	}
+
+}
+
+TSharedPtr<FAutomationTestExcludeOptions> FAutomationReport::GetExcludeOptions()
+{
+	ExcludeTestInfo.Test = *TestInfo.GetFullTestPath();
+	return ExcludeTestInfo.GetOptions();
+}
+
diff --git a/Engine/Source/Developer/AutomationController/Private/AutomationReport.h b/Engine/Source/Developer/AutomationController/Private/AutomationReport.h
--- a/Engine/Source/Developer/AutomationController/Private/AutomationReport.h
+++ b/Engine/Source/Developer/AutomationController/Private/AutomationReport.h
@@ -4,6 +4,7 @@
 
 #include "CoreMinimal.h"
 #include "Misc/AutomationTest.h"
+#include "AutomationTestExcludelist.h"
 #include "IAutomationReport.h"
 
 /**
@@ -69,6 +70,10 @@ public:
 	virtual void ResetNetworkCommandResponses() override;
 	virtual const bool ExpandInUI() const override;
 	virtual void StopRunningTest() override;
+	virtual bool IsToBeSkipped(FName* OutReason = nullptr, bool* OutWarn = nullptr) const override;
+	virtual bool IsToBeSkippedByPropagation() const override;
+	virtual void SetSkipFlag(bool bEnableSkip, const FAutomationTestExcludelistEntry* Template = nullptr, bool bFromPropagation = false) override;
+	virtual TSharedPtr<FAutomationTestExcludeOptions> GetExcludeOptions() override;
 
 private:
 
@@ -104,4 +109,10 @@ private:
 
 	/** Structure holding the test info */
 	FAutomationTestInfo TestInfo;
+
+	/** Structure holding the info on the exclude test */
+	FAutomationTestExcludelistEntry ExcludeTestInfo;
+	
+	/** True if this item is inside the exclude list */
+	bool bNeedToSkip = false;
 };
diff --git a/Engine/Source/Developer/AutomationController/Private/AutomationReportManager.cpp b/Engine/Source/Developer/AutomationController/Private/AutomationReportManager.cpp
--- a/Engine/Source/Developer/AutomationController/Private/AutomationReportManager.cpp
+++ b/Engine/Source/Developer/AutomationController/Private/AutomationReportManager.cpp
@@ -201,9 +201,9 @@ void FAutomationReportManager::AddResultReport(TSharedPtr< IAutomationReport > I
 				{
 					Status = "Success";
 				}
-				else if( InReport->GetState( ClusterIndex,CurrentTestPass ) == EAutomationState::NotEnoughParticipants )
+				else if( InReport->GetState( ClusterIndex,CurrentTestPass ) == EAutomationState::Skipped )
 				{
-					Status = TEXT( "Could not run." );
+					Status = TEXT( "Skipped." );
 				}
 				else
 				{
diff --git a/Engine/Source/Developer/AutomationController/Public/IAutomationReport.h b/Engine/Source/Developer/AutomationController/Public/IAutomationReport.h
--- a/Engine/Source/Developer/AutomationController/Public/IAutomationReport.h
+++ b/Engine/Source/Developer/AutomationController/Public/IAutomationReport.h
@@ -4,6 +4,8 @@
 
 #include "CoreMinimal.h"
 #include "Misc/AutomationTest.h"
+#include "AutomationTestExcludelist.h"
+#include "AutomationState.h"
 #include "UObject/ObjectMacros.h"
 #include "IAutomationReport.generated.h"
 
@@ -21,50 +23,6 @@ typedef TSharedPtr<class IAutomationReport> IAutomationReportPtr;
 typedef TSharedRef<class IAutomationReport> IAutomationReportRef;
 
 
-/** Enumeration of unit test status for special dialog */
-UENUM()
-enum class EAutomationState : uint8
-{
-	NotRun,					// Automation test was not run
-	InProcess,				// Automation test is running now
-	Fail,					// Automation test was run and failed
-	Success,				// Automation test was run and succeeded
-	NotEnoughParticipants,	// Automation test was not run due to number of participants
-};
-
-
-inline const TCHAR* ToString(EAutomationState InType)
-{
-	switch (InType)
-	{
-		case (EAutomationState::NotRun) :
-		{
-			return TEXT("NotRun");
-		}
-		case (EAutomationState::Fail) :
-		{
-			return TEXT("Fail");
-		}
-		case (EAutomationState::Success) :
-		{
-			return TEXT("Pass");
-		}
-		case (EAutomationState::InProcess) :
-		{
-			return TEXT("InProgress");
-		}
-		case (EAutomationState::NotEnoughParticipants) :
-		{
-			return TEXT("NotEnoughParticipants");
-		}
-		default:
-		{
-			return TEXT("Invalid");
-		}
-	}
-	return TEXT("Invalid");
-}
-
 UENUM()
 enum class EAutomationArtifactType : uint8
 {
@@ -522,6 +480,35 @@ public:
 	/** Stop the test which is creating this report. */
 	virtual void StopRunningTest() = 0;
 
+	/**
+	* Is the test need to be skipped.
+	*
+	* @return true if the test is inside exclude list.
+	*/
+	virtual bool IsToBeSkipped(FName* OutReason = nullptr, bool* OutWarn = nullptr) const = 0;
+
+	/**
+	* Is the test is to be skipped through propagation.
+	*
+	* @return true if the test is inside exclude list via a propagation.
+	*/
+	virtual bool IsToBeSkippedByPropagation() const = 0;
+	
+	/**
+	* Add or remove test from exclude list.
+	*
+	* @param if true, add this item to the exclude list, remove otherwise.
+	*/
+	virtual void SetSkipFlag(bool bEnableSkip, const FAutomationTestExcludelistEntry* Template = nullptr, bool bFromPropagation = false) = 0;
+
+	/**
+	* Produce exclude options .
+	*
+	* @return exclude options based on report exclude info.
+	*/
+	virtual TSharedPtr<FAutomationTestExcludeOptions> GetExcludeOptions() = 0;
+
+
 	// Event that allows log to refresh once a test has finished
 	DECLARE_DELEGATE_OneParam(FOnSetResultsEvent, TSharedPtr<IAutomationReport>);
 	FOnSetResultsEvent OnSetResults;
diff --git a/Engine/Source/Developer/AutomationWindow/AutomationWindow.Build.cs b/Engine/Source/Developer/AutomationWindow/AutomationWindow.Build.cs
--- a/Engine/Source/Developer/AutomationWindow/AutomationWindow.Build.cs
+++ b/Engine/Source/Developer/AutomationWindow/AutomationWindow.Build.cs
@@ -23,8 +23,9 @@ namespace UnrealBuildTool.Rules
                     "EditorStyle",
                     "CoreUObject",
                     "Json",
-                    "JsonUtilities"
-                }
+                    "JsonUtilities",
+					"AutomationTest"
+				}
 			);
 
             // Added more direct dependencies to the editor for testing functionality
@@ -34,7 +35,8 @@ namespace UnrealBuildTool.Rules
                     new string[] {
 						"EditorFramework",
                         "UnrealEd",
-                        "Engine", // Needed for UWorld/GWorld to find current level
+						"Kismet",
+						"Engine", // Needed for UWorld/GWorld to find current level
 				    }
                 );
             }
diff --git a/Engine/Source/Developer/AutomationWindow/Private/AutomationFilter.h b/Engine/Source/Developer/AutomationWindow/Private/AutomationFilter.h
--- a/Engine/Source/Developer/AutomationWindow/Private/AutomationFilter.h
+++ b/Engine/Source/Developer/AutomationWindow/Private/AutomationFilter.h
@@ -16,6 +16,7 @@ public:
 		: OnlySmokeTests( false )
 		, ShowErrors( false )
 		, ShowWarnings( false )
+		, OnlyExcludedTests ( false )
 	{ }
 
 public:
@@ -85,6 +86,27 @@ public:
 		return OnlySmokeTests;
 	}
 
+	/**
+	 * Should we only show excluded tests.
+	 *
+	 * @return True if we should be showing only excluded tests.
+	 * @see SetOnlyShowSmokeTests, ShouldShowErrors, ShouldShowWarnings
+	 */
+	const bool ShouldShowOnlyExcludedTests() const
+	{
+		return OnlyExcludedTests;
+	}
+
+	/**
+	 * Set if we should only show excluded tests.
+	 *
+	 * @param InOnlyExcludedTests If we should show excluded tests.
+	 */
+	void SetShowOnlyExcludedTests(const bool InOnlyExcludedTests)
+	{
+		OnlyExcludedTests = InOnlyExcludedTests;
+	}
+
 public:
 
 	// IFilter interface
@@ -95,17 +117,26 @@ public:
 	virtual bool PassesFilter( const TSharedPtr< IAutomationReport >& InReport ) const override
 	{
 		bool FilterPassed = true;
+
 		if (OnlySmokeTests)
 		{
 			//if we only want smoke tests and this isn't one
 			if (!InReport->IsSmokeTest())
 			{
-				FilterPassed = false;
+				return false;
 			}
 			//Leaf Nodes dictate this matching, not root nodes
 			if (InReport->GetTotalNumChildren())
 			{
-				FilterPassed = false;
+				return false;
+			}
+		}
+
+		if (OnlyExcludedTests)
+		{
+			if (!InReport->IsToBeSkipped())
+			{
+				return false;
 			}
 		}
 
@@ -139,4 +170,7 @@ private:
 
 	/** Only warnings will pass the test. */
 	bool ShowWarnings;
+
+	/** Only exclude tests. */
+	bool OnlyExcludedTests;
 };
diff --git a/Engine/Source/Developer/AutomationWindow/Private/SAutomationTestItem.cpp b/Engine/Source/Developer/AutomationWindow/Private/SAutomationTestItem.cpp
--- a/Engine/Source/Developer/AutomationWindow/Private/SAutomationTestItem.cpp
+++ b/Engine/Source/Developer/AutomationWindow/Private/SAutomationTestItem.cpp
@@ -12,6 +12,7 @@
 #include "Widgets/Text/STextBlock.h"
 #include "Widgets/Input/SCheckBox.h"
 #include "EditorStyleSet.h"
+#include "AutomationTestExcludelist.h"
 #include "SAutomationWindow.h"
 
 #if WITH_EDITOR
@@ -19,6 +20,8 @@
 	#include "EngineGlobals.h"
 	#include "Editor.h"
 	#include "AssetRegistryModule.h"
+	#include "Dialogs/Dialogs.h"
+	#include "SKismetInspector.h"
 #endif
 
 #include "Widgets/Input/SHyperlink.h"
@@ -249,6 +252,37 @@ TSharedRef<SWidget> SAutomationTestItem::GenerateWidgetForColumn( const FName& C
 		return SNew( STextBlock )
 		.Text( this, &SAutomationTestItem::ItemStatus_DurationText);
 	}
+	else if (ColumnName == AutomationTestWindowConstants::IsToBeSkipped)
+	{
+		return SNew(SHorizontalBox)
+			+ SHorizontalBox::Slot()
+			.HAlign(HAlign_Center)
+			.AutoWidth()
+			[
+				SNew(SCheckBox)
+				.IsChecked(this, &SAutomationTestItem::IsToBeSkipped)
+				.IsEnabled(this, &SAutomationTestItem::IsDirectlyExcluded)
+				.OnCheckStateChanged(this, &SAutomationTestItem::SetSkipFlag)
+				.ToolTipText(this, &SAutomationTestItem::GetExcludeReason)
+#if WITH_EDITOR
+			]
+		+ SHorizontalBox::Slot()
+			.HAlign(HAlign_Center)
+			.AutoWidth()
+			[
+				SNew(SButton)
+				.ButtonStyle(FAppStyle::Get(), "SimpleButton")
+				.Visibility(this, &SAutomationTestItem::IsDirectlyExcluded_GetVisibility)
+				.ToolTipText(LOCTEXT("EditExcludeOptions", "Edit exclude options"))
+				.OnClicked(FOnClicked::CreateSP(this, &SAutomationTestItem::OnEditExcludeOptionsClicked))
+				[
+					SNew(SImage)
+					.ColorAndOpacity(FSlateColor::UseForeground())
+				.Image(FAppStyle::Get().GetBrush("Icons.Edit"))
+				]
+#endif
+			];
+	}
 
 
 	return SNullWidget::NullWidget;
@@ -286,9 +320,9 @@ FText SAutomationTestItem::GetTestToolTip( int32 ClusterIndex ) const
 	{
 		TestToolTip = LOCTEXT("TestToolTipNotRun", "Not Run");
 	}
-	else if( TestState == EAutomationState::NotEnoughParticipants )
+	else if( TestState == EAutomationState::Skipped )
 	{
-		TestToolTip = LOCTEXT("ToolTipNotEnoughParticipants", "This test could not be completed as there were not enough participants.");
+		TestToolTip = LOCTEXT("ToolTipSkipped", "This test was skipped.");
 	}
 	else
 	{
@@ -317,6 +351,63 @@ ECheckBoxState SAutomationTestItem::IsTestEnabled() const
 	return TestStatus->IsEnabled() ? ECheckBoxState::Checked : ECheckBoxState::Unchecked;
 }
 
+ECheckBoxState SAutomationTestItem::IsToBeSkipped() const
+{
+	return TestStatus->IsToBeSkipped() ? ECheckBoxState::Checked : ECheckBoxState::Unchecked;
+}
+
+bool SAutomationTestItem::IsDirectlyExcluded() const
+{
+	return WITH_EDITOR && !TestStatus->IsToBeSkippedByPropagation();
+}
+
+EVisibility SAutomationTestItem::IsDirectlyExcluded_GetVisibility() const
+{
+	return  TestStatus->IsToBeSkipped() && IsDirectlyExcluded() ? EVisibility::Visible : EVisibility::Collapsed;
+}
+
+FText SAutomationTestItem::GetExcludeReason() const
+{
+	FName Reason;
+	bool IsToBeSkipped = TestStatus->IsToBeSkipped(&Reason);
+
+	return IsToBeSkipped ? FText::FromName(Reason) : LOCTEXT("", "");
+}
+
+void SAutomationTestItem::SetSkipFlag(ECheckBoxState Enable)
+{
+#if WITH_EDITOR
+	if (Enable == ECheckBoxState::Checked)
+	{
+		OnEditExcludeOptionsClicked();
+	}
+	else
+	{
+		TestStatus->SetSkipFlag(false);
+	}
+#endif
+}
+
+FReply SAutomationTestItem::OnEditExcludeOptionsClicked()
+{
+#if WITH_EDITOR
+	TSharedPtr<FAutomationTestExcludeOptions> Options = TestStatus->GetExcludeOptions();
+	TSharedPtr<FStructOnScope> StructToDisplay = MakeShareable(new FStructOnScope(FAutomationTestExcludeOptions::StaticStruct(), (uint8*)Options.Get()));
+
+	TSharedRef<SKismetInspector> KismetInspector = SNew(SKismetInspector);
+	KismetInspector->ShowSingleStruct(StructToDisplay);
+
+	SGenericDialogWidget::FArguments DialogArguments;
+	DialogArguments.OnOkPressed_Lambda([Options, this]()
+		{
+			auto Entry = FAutomationTestExcludelistEntry(*Options);
+			TestStatus->SetSkipFlag(true, &Entry, false);
+		});
+
+	SGenericDialogWidget::OpenDialog(LOCTEXT("ExcludeTestOptions", "Exclude Test Options"), KismetInspector, DialogArguments, true);
+#endif
+	return FReply::Handled();
+}
 
 FSlateColor SAutomationTestItem::ItemStatus_BackgroundColor(const int32 ClusterIndex) const
 {
@@ -485,8 +576,8 @@ const FSlateBrush* SAutomationTestItem::ItemStatus_StatusImage(const int32 Clust
 		}
 		break;
 
-	case EAutomationState::NotEnoughParticipants:
-		ImageToUse = FEditorStyle::GetBrush("Automation.NotEnoughParticipants");
+	case EAutomationState::Skipped:
+		ImageToUse = FEditorStyle::GetBrush("Automation.Skipped");
 		break;
 
 	default:
@@ -498,7 +589,6 @@ const FSlateBrush* SAutomationTestItem::ItemStatus_StatusImage(const int32 Clust
 	return ImageToUse;
 }
 
-
 /* SAutomationTestitem event handlers
  *****************************************************************************/
 
diff --git a/Engine/Source/Developer/AutomationWindow/Private/SAutomationTestItem.h b/Engine/Source/Developer/AutomationWindow/Private/SAutomationTestItem.h
--- a/Engine/Source/Developer/AutomationWindow/Private/SAutomationTestItem.h
+++ b/Engine/Source/Developer/AutomationWindow/Private/SAutomationTestItem.h
@@ -80,6 +80,46 @@ protected:
 	*/
 	ECheckBoxState IsTestEnabled( ) const;
 
+	/**
+	* Is the test inside exclude list. This info is taken from Config/DefaultEngine.ini, section [AutomationTestExcludelist].
+	*
+	* @return true if the test is inside the exclude list.
+	*/
+	ECheckBoxState IsToBeSkipped() const;
+
+	/**
+	* Is the test inside exclude list through direct exclusion (not through propagation).
+	*
+	* @return true if the test is inside the exclude list by direct exclusion.
+	*/
+	bool IsDirectlyExcluded() const;
+
+	/**
+	* Return the visibility state based on exclude state.
+	*/
+	EVisibility IsDirectlyExcluded_GetVisibility() const;
+
+	/**
+	* Reason for the exclusion
+	*
+	* @return the exclusion reason
+	*/
+	FText GetExcludeReason() const;
+
+	/**
+	* Add or remove test from exclude list. It will change Config/DefaultEngine.ini, section [AutomationTestExcludelist].
+	*
+	* @param if true, add this item to exclude list, remove otherwise.
+	*/
+	void SetSkipFlag(ECheckBoxState Enable);
+
+	/**
+	* Open Exclude test options editor
+	*
+	* @return respond of dialog
+	*/
+	FReply OnEditExcludeOptionsClicked();
+
 	/**
 	 * Returns color that indicates test status per cluster.
 	 *
diff --git a/Engine/Source/Developer/AutomationWindow/Private/SAutomationWindow.cpp b/Engine/Source/Developer/AutomationWindow/Private/SAutomationWindow.cpp
--- a/Engine/Source/Developer/AutomationWindow/Private/SAutomationWindow.cpp
+++ b/Engine/Source/Developer/AutomationWindow/Private/SAutomationWindow.cpp
@@ -67,7 +67,8 @@ public:
 		UI_COMMAND( ErrorFilter, "Errors", "Toggle Error Filter", EUserInterfaceActionType::ToggleButton, FInputChord() );
 		UI_COMMAND( WarningFilter, "Warnings", "Toggle Warning Filter", EUserInterfaceActionType::ToggleButton, FInputChord() );
 		UI_COMMAND( DeveloperDirectoryContent, "Dev Content", "Developer Directory Content Filter (when enabled, developer directories are also included)", EUserInterfaceActionType::ToggleButton, FInputChord() );
-		
+		UI_COMMAND( ExcludedTestsFilter, "Excluded Tests", "Toggle Excluded Tests only", EUserInterfaceActionType::ToggleButton, FInputChord());
+
 #if WITH_EDITOR
 		// Added button for running the currently open level test.
 		UI_COMMAND(RunLevelTest, "Run Level Test", "Run Level Test", EUserInterfaceActionType::Button, FInputChord());
@@ -79,6 +80,7 @@ public:
 	TSharedPtr<FUICommandInfo> ErrorFilter;
 	TSharedPtr<FUICommandInfo> WarningFilter;
 	TSharedPtr<FUICommandInfo> DeveloperDirectoryContent;
+	TSharedPtr<FUICommandInfo> ExcludedTestsFilter;
 
 #if WITH_EDITOR
 	TSharedPtr<FUICommandInfo> RunLevelTest;
@@ -248,12 +250,16 @@ void SAutomationWindow::Construct( const FArguments& InArgs, const IAutomationCo
 		.DefaultLabel(LOCTEXT("TestDurationRange", "Duration"))
 
 		+ SHeaderRow::Column( AutomationTestWindowConstants::Status )
-		.FillWidth(0.20f)
+		.FixedWidth(50.0f)
 		[
 			//platform header placeholder
 			PlatformsHBox.ToSharedRef()
 		]
 
+		+ SHeaderRow::Column(AutomationTestWindowConstants::IsToBeSkipped)
+		.FillWidth(0.10f)
+		.DefaultLabel(LOCTEXT("Excluded", "Excluded"))
+
 		);
 
 	RequestedFilterComboList.Empty();
@@ -606,6 +612,12 @@ void SAutomationWindow::CreateCommands()
 		FIsActionChecked::CreateRaw( this, &SAutomationWindow::IsDeveloperDirectoryIncluded )
 		);
 
+	ActionList.MapAction( Commands.ExcludedTestsFilter,
+		FExecuteAction::CreateRaw( this, &SAutomationWindow::OnToggleExcludedTestsFilter ),
+		FCanExecuteAction::CreateRaw( this, &SAutomationWindow::IsAutomationControllerIdle ),
+		FIsActionChecked::CreateRaw( this, &SAutomationWindow::IsExcludedTestsFilterOn )
+		);
+
 	// Added button for running the currently open level test.
 #if WITH_EDITOR
 	ActionList.MapAction(Commands.RunLevelTest,
@@ -658,6 +670,7 @@ TSharedRef< SWidget > SAutomationWindow::MakeAutomationWindowToolBar( const TSha
 				ToolbarBuilder.AddToolBarButton( FAutomationWindowCommands::Get().ErrorFilter );
 				ToolbarBuilder.AddToolBarButton( FAutomationWindowCommands::Get().WarningFilter );
 				ToolbarBuilder.AddToolBarButton( FAutomationWindowCommands::Get().DeveloperDirectoryContent );
+				ToolbarBuilder.AddToolBarButton( FAutomationWindowCommands::Get().ExcludedTestsFilter);
 			}
 			ToolbarBuilder.EndSection();
 			ToolbarBuilder.BeginSection("GroupFlags");
@@ -1936,6 +1949,18 @@ void SAutomationWindow::OnToggleDeveloperDirectoryIncluded()
 	ListTests();
 }
 
+bool SAutomationWindow::IsExcludedTestsFilterOn() const
+{
+	return AutomationGeneralFilter->ShouldShowOnlyExcludedTests();
+}
+
+
+void SAutomationWindow::OnToggleExcludedTestsFilter()
+{
+	AutomationGeneralFilter->SetShowOnlyExcludedTests(!IsExcludedTestsFilterOn());
+	OnRefreshTestCallback();
+}
+
 
 bool SAutomationWindow::IsSmokeTestFilterOn() const
 {
@@ -2120,7 +2145,6 @@ void SAutomationWindow::ChangeTheSelectionToThisRow(TSharedPtr< IAutomationRepor
 	TestTable->SetSelection(ThisRow, ESelectInfo::Direct);
 }
 
-
 bool SAutomationWindow::IsRowSelected(TSharedPtr< IAutomationReport >  ThisRow)
 {
 	TArray< TSharedPtr<IAutomationReport> >SelectedReport = TestTable->GetSelectedItems();
diff --git a/Engine/Source/Developer/AutomationWindow/Private/SAutomationWindow.h b/Engine/Source/Developer/AutomationWindow/Private/SAutomationWindow.h
--- a/Engine/Source/Developer/AutomationWindow/Private/SAutomationWindow.h
+++ b/Engine/Source/Developer/AutomationWindow/Private/SAutomationWindow.h
@@ -38,6 +38,7 @@ namespace AutomationTestWindowConstants
 	const FName RequiredDeviceCount( TEXT("RequiredDeviceCount") );
 	const FName Status( TEXT("Status") );
 	const FName Timing( TEXT("Timing") );
+	const FName IsToBeSkipped(TEXT("ToBeSkipped"));
 }
 
 
@@ -283,10 +284,16 @@ private:
 	
 	/** Returns if we're considering tests on content within the developer folders */
 	bool IsDeveloperDirectoryIncluded() const;
-	
+
 	/** Toggles the consideration of tests within developer folders */
 	void OnToggleDeveloperDirectoryIncluded();
-	
+
+	/** Returns if we're considering tests from exclude list only */
+	bool IsExcludedTestsFilterOn() const;
+
+	/** Toggles the consideration of tests from exclude list only */
+	void OnToggleExcludedTestsFilter();
+
 	/** Returns if we're filtering based on if the test is a "smoke" test */
 	bool IsSmokeTestFilterOn() const;
 	
diff --git a/Engine/Source/Editor/EditorStyle/Private/SlateEditorStyle.cpp b/Engine/Source/Editor/EditorStyle/Private/SlateEditorStyle.cpp
--- a/Engine/Source/Editor/EditorStyle/Private/SlateEditorStyle.cpp
+++ b/Engine/Source/Editor/EditorStyle/Private/SlateEditorStyle.cpp
@@ -7864,7 +7864,7 @@ void FSlateEditorStyle::FStyle::SetupAutomationStyles()
 		Set( "Automation.Fail", new IMAGE_BRUSH( "Automation/Fail", Icon16x16 ) );
 		Set( "Automation.InProcess", new IMAGE_BRUSH( "Automation/InProcess", Icon16x16 ) );
 		Set( "Automation.NotRun", new IMAGE_BRUSH( "Automation/NotRun", Icon16x16, FLinearColor(0.0f, 0.0f, 0.0f, 0.4f) ) );
-		Set( "Automation.NotEnoughParticipants", new IMAGE_BRUSH( "Automation/NotEnoughParticipants", Icon16x16 ) );
+		Set( "Automation.Skipped", new IMAGE_BRUSH( "Automation/NoSessionWarning", Icon16x16 ) );
 		Set( "Automation.ParticipantsWarning", new IMAGE_BRUSH( "Automation/ParticipantsWarning", Icon16x16 ) );
 		Set( "Automation.Participant", new IMAGE_BRUSH( "Automation/Participant", Icon16x16 ) );
 		
@@ -7887,6 +7887,7 @@ void FSlateEditorStyle::FStyle::SetupAutomationStyles()
 		Set( "AutomationWindow.WarningFilter", new IMAGE_BRUSH( "Automation/WarningFilter", Icon40x40) );
 		Set( "AutomationWindow.SmokeTestFilter", new IMAGE_BRUSH( "Automation/SmokeTestFilter", Icon40x40) );
 		Set( "AutomationWindow.DeveloperDirectoryContent", new IMAGE_BRUSH( "Automation/DeveloperDirectoryContent", Icon40x40) );
+		Set( "AutomationWindow.ExcludedTestsFilter", new IMAGE_BRUSH("Automation/ExcludedTestsFilter", Icon40x40) );
 		Set( "AutomationWindow.ErrorFilter.Small", new IMAGE_BRUSH( "Automation/ErrorFilter", Icon20x20) );
 		Set( "AutomationWindow.WarningFilter.Small", new IMAGE_BRUSH( "Automation/WarningFilter", Icon20x20) );
 		Set( "AutomationWindow.SmokeTestFilter.Small", new IMAGE_BRUSH( "Automation/SmokeTestFilter", Icon20x20) );
diff --git a/Engine/Source/Editor/EditorStyle/Private/StarshipStyle.cpp b/Engine/Source/Editor/EditorStyle/Private/StarshipStyle.cpp
--- a/Engine/Source/Editor/EditorStyle/Private/StarshipStyle.cpp
+++ b/Engine/Source/Editor/EditorStyle/Private/StarshipStyle.cpp
@@ -7434,7 +7434,7 @@ void FStarshipEditorStyle::FStyle::SetupAutomationStyles()
 		Set( "Automation.Fail", new IMAGE_BRUSH( "Automation/Fail", Icon16x16 ) );
 		Set( "Automation.InProcess", new IMAGE_BRUSH( "Automation/InProcess", Icon16x16 ) );
 		Set( "Automation.NotRun", new IMAGE_BRUSH( "Automation/NotRun", Icon16x16, FLinearColor(0.0f, 0.0f, 0.0f, 0.4f) ) );
-		Set( "Automation.NotEnoughParticipants", new IMAGE_BRUSH( "Automation/NotEnoughParticipants", Icon16x16 ) );
+		Set( "Automation.Skipped", new IMAGE_BRUSH( "Automation/NoSessionWarning", Icon16x16 ) );
 		Set( "Automation.ParticipantsWarning", new IMAGE_BRUSH( "Automation/ParticipantsWarning", Icon16x16 ) );
 		Set( "Automation.Participant", new IMAGE_BRUSH( "Automation/Participant", Icon16x16 ) );
 		
@@ -7457,6 +7457,7 @@ void FStarshipEditorStyle::FStyle::SetupAutomationStyles()
 		Set( "AutomationWindow.WarningFilter", new IMAGE_BRUSH( "Automation/WarningFilter", Icon40x40) );
 		Set( "AutomationWindow.SmokeTestFilter", new IMAGE_BRUSH( "Automation/SmokeTestFilter", Icon40x40) );
 		Set( "AutomationWindow.DeveloperDirectoryContent", new IMAGE_BRUSH( "Automation/DeveloperDirectoryContent", Icon40x40) );
+		Set( "AutomationWindow.ExcludedTestsFilter", new IMAGE_BRUSH("Automation/ExcludedTestsFilter", Icon40x40) );
 		Set( "AutomationWindow.ErrorFilter.Small", new IMAGE_BRUSH( "Automation/ErrorFilter", Icon20x20) );
 		Set( "AutomationWindow.WarningFilter.Small", new IMAGE_BRUSH( "Automation/WarningFilter", Icon20x20) );
 		Set( "AutomationWindow.SmokeTestFilter.Small", new IMAGE_BRUSH( "Automation/SmokeTestFilter", Icon20x20) );
diff --git a/Engine/Source/Programs/AutomationTool/Gauntlet/Framework/Base/Gauntlet.TestReport.cs b/Engine/Source/Programs/AutomationTool/Gauntlet/Framework/Base/Gauntlet.TestReport.cs
--- a/Engine/Source/Programs/AutomationTool/Gauntlet/Framework/Base/Gauntlet.TestReport.cs
+++ b/Engine/Source/Programs/AutomationTool/Gauntlet/Framework/Base/Gauntlet.TestReport.cs
@@ -34,7 +34,7 @@ namespace Gauntlet
 		InProcess,
 		Fail,
 		Success,
-		NotEnoughParticipants
+		Skipped
 	}
 
 	/// <summary>
@@ -48,13 +48,21 @@ namespace Gauntlet
 		string Type { get; }
 
 		/// <summary>
-		/// Set a property of the test report type
+		/// Set a property of the test report
 		/// </summary>
 		/// <param name="Property"></param>
 		/// <param name="Value"></param>
 		/// <returns></returns>
 		void SetProperty(string Property, object Value);
 
+		/// <summary>
+		/// Set a metadata key/value to the test report
+		/// </summary>
+		/// <param name="Key"></param>
+		/// <param name="Value"></param>
+		/// <returns></returns>
+		void SetMetadata(string Key, string Value);
+
 		/// <summary>
 		/// Add event to the test report
 		/// </summary>
@@ -123,6 +131,15 @@ namespace Gauntlet
 	/// </summary>
 	public abstract class BaseTestReport : ITestReport, ITelemetryReport
 	{
+		public BaseTestReport()
+		{
+			Metadata = new Dictionary<string, string>();
+		}
+		/// <summary>
+		/// Metadata blackboard
+		/// </summary>
+		public Dictionary<string, string> Metadata { get; set; }
+
 		/// <summary>
 		/// Return report type
 		/// </summary>
@@ -134,7 +151,7 @@ namespace Gauntlet
 		protected List<TelemetryData> TelemetryDataList;
 
 		/// <summary>
-		/// Set a property of the test report type
+		/// Set a property of the test report
 		/// </summary>
 		/// <param name="Property"></param>
 		/// <param name="Value"></param>
@@ -145,6 +162,17 @@ namespace Gauntlet
 			PropertyInstance.SetValue(this, Convert.ChangeType(Value, PropertyInstance.PropertyType));
 		}
 
+		/// <summary>
+		/// Set a metadata key/value to the test report
+		/// </summary>
+		/// <param name="Key"></param>
+		/// <param name="Value"></param>
+		/// <returns></returns>
+		public virtual void SetMetadata(string Key, string Value)
+		{
+			Metadata[Key] = Value;
+		}
+
 		/// <summary>
 		/// Add event to the test report
 		/// </summary>
diff --git a/Engine/Source/Programs/AutomationTool/Gauntlet/Framework/Utils/Gauntlet.HordeReport.cs b/Engine/Source/Programs/AutomationTool/Gauntlet/Framework/Utils/Gauntlet.HordeReport.cs
--- a/Engine/Source/Programs/AutomationTool/Gauntlet/Framework/Utils/Gauntlet.HordeReport.cs
+++ b/Engine/Source/Programs/AutomationTool/Gauntlet/Framework/Utils/Gauntlet.HordeReport.cs
@@ -131,6 +131,23 @@ namespace Gauntlet
 			}
 		}
 
+		/// <summary>
+		/// Contains detailed information about device that run tests
+		/// </summary>
+		public class Device
+		{
+			public string DeviceName { get; set; }
+			public string Instance { get; set; }
+			public string Platform { get; set; }
+			public string OSVersion { get; set; }
+			public string Model { get; set; }
+			public string GPU { get; set; }
+			public string CPUModel { get; set; }
+			public int RAMInGB { get; set; }
+			public string RenderMode { get; set; }
+			public string RHI { get; set; }
+		}
+
 		/// <summary>
 		/// Contains reference to files used or generated for file comparison
 		/// </summary>
@@ -194,6 +211,7 @@ namespace Gauntlet
 			public string TestDisplayName { get; set; }
 			public string FullTestPath { get; set; }
 			public TestStateType State { get; set; }
+			public string DeviceInstance { get; set; }
 			public int Warnings { get; set; }
 			public int Errors { get; set; }
 			public List<Artifact> Artifacts { get; set; }
@@ -246,6 +264,11 @@ namespace Gauntlet
 				get { return TestDetailed.State; }
 				set { TestDetailed.State = value; }
 			}
+			public string DeviceInstance
+			{
+				get { return TestDetailed.DeviceInstance; }
+				set { TestDetailed.DeviceInstance = value; }
+			}
 			public int Errors
 			{
 				get { return TestDetailed.Errors; }
@@ -288,12 +311,13 @@ namespace Gauntlet
 				get { return "Unreal Automated Tests"; }
 			}
 
-			public UnrealEngineTestPassResults()
+			public UnrealEngineTestPassResults() : base()
 			{
+				Devices = new List<Device>();
 				Tests = new List<TestResult>();
 			}
 
-			public string ClientDescriptor { get; set; }
+			public List<Device> Devices { get; set; }
 			public string ReportCreatedOn { get; set; }
 			public string ReportURL { get; set; }
 			public int SucceededCount { get; set; }
@@ -304,6 +328,17 @@ namespace Gauntlet
 			public float TotalDurationSeconds { get; set; }
 			public List<TestResult> Tests { get; set; }
 
+			/// <summary>
+			/// Add a new Device to the pass results and return it 
+			/// </summary>
+			private Device AddNewDevice()
+			{
+				Device NewDevice = new Device();
+				Devices.Add(NewDevice);
+
+				return NewDevice;
+			}
+
 			/// <summary>
 			/// Add a new TestResult to the pass results and return it 
 			/// </summary>
@@ -329,7 +364,23 @@ namespace Gauntlet
 			public static UnrealEngineTestPassResults FromUnrealAutomatedTests(UnrealAutomatedTestPassResults InTestPassResults, string ReportPath, string ReportURL)
 			{
 				UnrealEngineTestPassResults OutTestPassResults = new UnrealEngineTestPassResults();
-				OutTestPassResults.ClientDescriptor = InTestPassResults.ClientDescriptor;
+				if (InTestPassResults.Devices != null)
+				{
+					foreach (UnrealAutomationDevice InDevice in InTestPassResults.Devices)
+					{
+						Device ConvertedDevice = OutTestPassResults.AddNewDevice();
+						ConvertedDevice.DeviceName = InDevice.DeviceName;
+						ConvertedDevice.Instance = InDevice.Instance;
+						ConvertedDevice.Platform = InDevice.Platform;
+						ConvertedDevice.OSVersion = InDevice.OSVersion;
+						ConvertedDevice.Model = InDevice.Model;
+						ConvertedDevice.GPU = InDevice.GPU;
+						ConvertedDevice.CPUModel = InDevice.CPUModel;
+						ConvertedDevice.RAMInGB = InDevice.RAMInGB;
+						ConvertedDevice.RenderMode = InDevice.RenderMode;
+						ConvertedDevice.RHI = InDevice.RHI;
+					}
+				}
 				OutTestPassResults.ReportCreatedOn = InTestPassResults.ReportCreatedOn;
 				OutTestPassResults.ReportURL = ReportURL;
 				OutTestPassResults.SucceededCount = InTestPassResults.Succeeded;
@@ -346,6 +397,7 @@ namespace Gauntlet
 						ConvertedTestResult.TestDisplayName = InTestResult.TestDisplayName;
 						ConvertedTestResult.FullTestPath = InTestResult.FullTestPath;
 						ConvertedTestResult.State = InTestResult.State;
+						ConvertedTestResult.DeviceInstance = InTestResult.DeviceInstance;
 						Guid TestGuid = Guid.NewGuid();
 						ConvertedTestResult.ArtifactName = TestGuid + ".json";
 						InTestResult.ArtifactName = ConvertedTestResult.ArtifactName;
@@ -429,7 +481,7 @@ namespace Gauntlet
 			{
 				get { return "Simple Report"; }
 			}
-			public SimpleTestReport()
+			public SimpleTestReport() : base()
 			{
 				Logs = new List<String>();
 				Errors = new List<String>();
diff --git a/Engine/Source/Programs/AutomationTool/Gauntlet/Unreal/Automation/UE.Automation.cs b/Engine/Source/Programs/AutomationTool/Gauntlet/Unreal/Automation/UE.Automation.cs
--- a/Engine/Source/Programs/AutomationTool/Gauntlet/Unreal/Automation/UE.Automation.cs
+++ b/Engine/Source/Programs/AutomationTool/Gauntlet/Unreal/Automation/UE.Automation.cs
@@ -607,9 +607,9 @@ namespace UE
 				UnrealAutomatedTestPassResults JsonTestPassResults = GetTestPassResults(EditorRole.LogSummary);
 
 				// Filter our tests into categories
-				IEnumerable<UnrealAutomatedTestResult> AllTests = JsonTestPassResults.Tests;
+				IEnumerable<UnrealAutomatedTestResult> AllTests = JsonTestPassResults.Tests.Where(T => !T.WasSkipped);
 				IEnumerable<UnrealAutomatedTestResult> IncompleteTests = AllTests.Where(T => !T.IsComplete);
-				IEnumerable<UnrealAutomatedTestResult> FailedTests = AllTests.Where(T => T.IsComplete && !T.HasSucceeded);
+				IEnumerable<UnrealAutomatedTestResult> FailedTests = AllTests.Where(T => T.IsComplete && T.HasFailed);
 				IEnumerable<UnrealAutomatedTestResult> TestsWithWarnings = AllTests.Where(T => T.HasSucceeded && T.HasWarnings);
 
 				Func<string, string> ErrorLineForHorde = (L) => string.Format("Err: {0}", L);
@@ -791,7 +791,7 @@ namespace UE
 				{
 					AutomationLogParser Parser = new AutomationLogParser(Role.LogSummary.FullLogContent);
 					AllErrors.AddRange(
-						Parser.GetResults().Where(R => !R.HasSucceeded)
+						Parser.GetResults().Where(R => R.HasFailed)
 							.SelectMany(R => R.Entries
 								.Where(E => E.Event.Type == EventType.Error)
 								.Distinct().Select(E => string.Format("[test={0}] {1}", R.TestDisplayName, E))
diff --git a/Engine/Source/Programs/AutomationTool/Gauntlet/Unreal/Base/Gauntlet.UnrealTestNode.cs b/Engine/Source/Programs/AutomationTool/Gauntlet/Unreal/Base/Gauntlet.UnrealTestNode.cs
--- a/Engine/Source/Programs/AutomationTool/Gauntlet/Unreal/Base/Gauntlet.UnrealTestNode.cs
+++ b/Engine/Source/Programs/AutomationTool/Gauntlet/Unreal/Base/Gauntlet.UnrealTestNode.cs
@@ -1182,6 +1182,9 @@ namespace Gauntlet
 					}
 				}
 			}
+			// Metadata
+			SetReportMetadata(HordeTestReport, GetConfiguration().RequiredRoles.Values.SelectMany(V => V));
+
 			return HordeTestReport;
 		}
 
@@ -1261,6 +1264,10 @@ namespace Gauntlet
 				// Convert test results for Horde
 				HordeReport.UnrealEngineTestPassResults HordeTestPassResults = HordeReport.UnrealEngineTestPassResults.FromUnrealAutomatedTests(JsonTestPassResults, UnrealAutomatedTestReportPath, ReportURL);
 				HordeTestPassResults.CopyTestResultsArtifacts(HordeArtifactPath);
+				// Metadata
+				// With UE Test Automation, we care only for one role.
+				var MainRole = new List<UnrealTestRole>() { GetConfiguration().GetMainRequiredRole() };
+				SetReportMetadata(HordeTestPassResults, MainRole);
 				// Attached test Artifacts
 				if (SessionArtifacts != null)
 				{
@@ -1279,6 +1286,20 @@ namespace Gauntlet
 			}
 		}
 
+		/// <summary>
+		/// Set Metadata on ITestReport
+		/// </summary>
+		/// <param name="Report"></param>
+		protected virtual void SetReportMetadata(ITestReport Report, IEnumerable<UnrealTestRole> Roles)
+		{
+			var AllRoleTypes = Roles.Select(R =>  R.Type);
+			var AllRoleContexts = Roles.Select(R => Context.GetRoleContext(R.Type));
+			Report.SetMetadata("Platform", string.Join("+", AllRoleContexts.Select(R => R.Platform.ToString()).Distinct().OrderBy(P => P)));
+			Report.SetMetadata("BuildType", string.Join("+", AllRoleTypes.Select(R => R.ToString()).Distinct().OrderBy(B => B)));
+			Report.SetMetadata("Configuration", string.Join("+", AllRoleContexts.Select(R => R.Configuration.ToString()).Distinct().OrderBy(C => C)));
+			Report.SetMetadata("Project", Context.BuildInfo.ProjectName);
+		}
+
 		/// <summary>
 		/// DEPRECATED Optional function that is called on test completion and gives an opportunity to submit a report to a Dashboard
 		/// </summary>
diff --git a/Engine/Source/Programs/AutomationTool/Gauntlet/Unreal/Engine/UnrealTests.cs b/Engine/Source/Programs/AutomationTool/Gauntlet/Unreal/Engine/UnrealTests.cs
--- a/Engine/Source/Programs/AutomationTool/Gauntlet/Unreal/Engine/UnrealTests.cs
+++ b/Engine/Source/Programs/AutomationTool/Gauntlet/Unreal/Engine/UnrealTests.cs
@@ -76,7 +76,7 @@ namespace Gauntlet
 					{
 						AutomationLogParser Parser = new AutomationLogParser(Role.LogSummary.FullLogContent);
 						AllErrors.AddRange(
-							Parser.GetResults().Where(R => !R.HasSucceeded)
+							Parser.GetResults().Where(R => R.HasFailed)
 							.SelectMany(R => R.Entries
 								.Where(E => E.Event.Type == EventType.Error).Select(E => string.Format("[test={0}] {1}", R.TestDisplayName, E))
 								)
@@ -220,7 +220,7 @@ namespace Gauntlet
 						AutomationLogParser Parser = new AutomationLogParser(InLog.FullLogContent);
 
 						IEnumerable<UnrealAutomatedTestResult> TotalTests = Parser.GetResults();
-						IEnumerable<UnrealAutomatedTestResult> FailedTests = TotalTests.Where(R => !R.HasSucceeded);
+						IEnumerable<UnrealAutomatedTestResult> FailedTests = TotalTests.Where(R => R.HasFailed);
 
 						// Tests failed so list that as our primary cause of failure
 						if (FailedTests.Any())
@@ -258,7 +258,7 @@ namespace Gauntlet
 					AutomationLogParser Parser = new AutomationLogParser(EditorRole.LogSummary.FullLogContent);
 
 					IEnumerable<UnrealAutomatedTestResult> AllTests = Parser.GetResults();
-					IEnumerable<UnrealAutomatedTestResult> FailedTests = AllTests.Where(R => R.IsComplete && !R.HasSucceeded);
+					IEnumerable<UnrealAutomatedTestResult> FailedTests = AllTests.Where(R => R.IsComplete && R.HasFailed);
 					IEnumerable<UnrealAutomatedTestResult> IncompleteTests = AllTests.Where(R => !R.IsComplete);
 
 					if (AllTests.Count() == 0)
diff --git a/Engine/Source/Programs/AutomationTool/Gauntlet/Unreal/Utils/Gauntlet.AutomationLogParser.cs b/Engine/Source/Programs/AutomationTool/Gauntlet/Unreal/Utils/Gauntlet.AutomationLogParser.cs
--- a/Engine/Source/Programs/AutomationTool/Gauntlet/Unreal/Utils/Gauntlet.AutomationLogParser.cs
+++ b/Engine/Source/Programs/AutomationTool/Gauntlet/Unreal/Utils/Gauntlet.AutomationLogParser.cs
@@ -80,10 +80,10 @@ namespace Gauntlet
 		/// <returns></returns>
 		public IEnumerable<UnrealAutomatedTestResult> GetResults()
 		{
-			IEnumerable<Match> TestStarts = Parser.GetAllMatches(@"LogAutomationController.+Test Started. Name={(.+?)}");
+			IEnumerable<Match> TestStarts = Parser.GetAllMatches(@"LogAutomationController.+Test Started. Name={(.+?)}\s+Path={(.+?)}");
 
 			// Find all automation results that succeeded/failed
-			// [00:10:54.148]   LogAutomationController: Display: Test Started. Name={ST_PR04}
+			// [00:10:54.148]   LogAutomationController: Display: Test Started. Name={ST_PR04} Path={Project.Functional Tests./Game/Tests/Rendering/PlanarReflection.ST_PR04}
 			// [2019.04.30-18.49.51:329][244]LogAutomationController: Display: Test Completed With Success. Name={ST_PR04} Path={Project.Functional Tests./Game/Tests/Rendering/PlanarReflection.ST_PR04}
 			// [2019.04.30-18.49.51:330] [244] LogAutomationController: BeginEvents: Project.Functional Tests./Game/Tests/Rendering/PlanarReflection.ST_PR04
 			// [2019.04.30 - 18.49.51:331][244] LogAutomationController: Screenshot 'ST_PR04' was similar!  Global Difference = 0.001377, Max Local Difference = 0.037953
@@ -106,17 +106,25 @@ namespace Gauntlet
 			{
 				UnrealAutomatedTestResult Result = new UnrealAutomatedTestResult();
 				string DisplayName = TestMatch.Groups[1].ToString();
-				string LongName = "";
-				bool Completed = false;
-				bool Passed = false;
+				string LongName = TestMatch.Groups[2].ToString(); ;
+				TestStateType State = TestStateType.InProcess;
 				
-				Match ResultMatch = TestResults.Where(M => M.Groups[2].ToString() == DisplayName).FirstOrDefault();
+				Match ResultMatch = TestResults.Where(M => M.Groups[3].ToString() == LongName).FirstOrDefault();
 
 				if (ResultMatch != null)
 				{
-					Completed = true;
-					Passed = ResultMatch.Groups[1].ToString().ToLower() == "passed" ? true : false;
-					LongName = ResultMatch.Groups[3].ToString();					
+					switch(ResultMatch.Groups[1].ToString().ToLower())
+					{
+					case "skipped":
+						State = TestStateType.Skipped;
+						break;
+					case "success":
+						State = TestStateType.Success;
+						break;
+					default:
+						State = TestStateType.Fail;
+						break;
+					}
 
 					string EventName = string.Format("BeginEvents: {0}", LongName);
 					int EventIndex = Array.FindIndex(AutomationChannel, S => S.Contains(EventName)) + 1;
@@ -145,7 +153,7 @@ namespace Gauntlet
 
 				Result.TestDisplayName = DisplayName;
 				Result.FullTestPath = LongName;
-				Result.State = Passed ? TestStateType.Success : Completed ? TestStateType.Fail : TestStateType.InProcess;
+				Result.State = State;
 
 				return Result;
 			});
diff --git a/Engine/Source/Programs/AutomationTool/Gauntlet/Unreal/Utils/Gauntlet.UnrealTestResults.cs b/Engine/Source/Programs/AutomationTool/Gauntlet/Unreal/Utils/Gauntlet.UnrealTestResults.cs
--- a/Engine/Source/Programs/AutomationTool/Gauntlet/Unreal/Utils/Gauntlet.UnrealTestResults.cs
+++ b/Engine/Source/Programs/AutomationTool/Gauntlet/Unreal/Utils/Gauntlet.UnrealTestResults.cs
@@ -9,6 +9,19 @@ using System.Linq;
 
 namespace Gauntlet
 {
+	public class UnrealAutomationDevice
+	{
+		public string DeviceName { get; set; }
+		public string Instance { get; set; }
+		public string Platform { get; set; }
+		public string OSVersion { get; set; }
+		public string Model { get; set; }
+		public string GPU { get; set; }
+		public string CPUModel { get; set; }
+		public int RAMInGB { get; set; }
+		public string RenderMode { get; set; }
+		public string RHI { get; set; }
+	}
 	public class UnrealAutomationComparisonFiles
 	{
 		public string Difference { get; set; }
@@ -70,6 +83,7 @@ namespace Gauntlet
 		public string FullTestPath { get; set; }
 		public string ArtifactName { get; set; }
 		public TestStateType State { get; set; }
+		public String DeviceInstance { get; set; }
 		public int Warnings { get; set; }
 		public int Errors { get; set; }
 		public List<UnrealAutomationArtifact> Artifacts { get; set; }
@@ -117,7 +131,7 @@ namespace Gauntlet
 		{
 			get
 			{
-				return State != TestStateType.InProcess;
+				return State != TestStateType.InProcess && State != TestStateType.NotRun;
 			}
 		}
 		[JsonIgnore]
@@ -129,6 +143,22 @@ namespace Gauntlet
 			}
 		}
 		[JsonIgnore]
+		public bool HasFailed
+		{
+			get
+			{
+				return State == TestStateType.Fail;
+			}
+		}
+		[JsonIgnore]
+		public bool WasSkipped
+		{
+			get
+			{
+				return State == TestStateType.Skipped;
+			}
+		}
+		[JsonIgnore]
 		public bool HasWarnings
 		{
 			get
@@ -163,7 +193,7 @@ namespace Gauntlet
 	}
 	public class UnrealAutomatedTestPassResults
 	{
-		public string ClientDescriptor { get; set; }
+		public List<UnrealAutomationDevice> Devices { get; set; }
 		public string ReportCreatedOn { get; set; }
 		public int Succeeded { get; set; }
 		public int SucceededWithWarnings { get; set; }
@@ -177,6 +207,7 @@ namespace Gauntlet
 
 		public UnrealAutomatedTestPassResults()
 		{
+			Devices = new List<UnrealAutomationDevice>();
 			Tests = new List<UnrealAutomatedTestResult>();
 		}
 
diff --git a/Engine/Source/Runtime/AutomationMessages/AutomationMessages.Build.cs b/Engine/Source/Runtime/AutomationMessages/AutomationMessages.Build.cs
--- a/Engine/Source/Runtime/AutomationMessages/AutomationMessages.Build.cs
+++ b/Engine/Source/Runtime/AutomationMessages/AutomationMessages.Build.cs
@@ -10,12 +10,13 @@ namespace UnrealBuildTool.Rules
 				new string[]
 				{
 					"CoreUObject",
+					"AutomationTest"
 				});
 
 			PrivateDependencyModuleNames.AddRange(
 				new string[]
 				{
-					"Core",
+					"Core"
 				});
 
 			PrivateIncludePaths.AddRange(
diff --git a/Engine/Source/Runtime/AutomationMessages/Public/AutomationWorkerMessages.h b/Engine/Source/Runtime/AutomationMessages/Public/AutomationWorkerMessages.h
--- a/Engine/Source/Runtime/AutomationMessages/Public/AutomationWorkerMessages.h
+++ b/Engine/Source/Runtime/AutomationMessages/Public/AutomationWorkerMessages.h
@@ -6,6 +6,7 @@
 #include "Containers/UnrealString.h"
 #include "Misc/Guid.h"
 #include "Misc/AutomationTest.h"
+#include "AutomationState.h"
 #include "UObject/ObjectMacros.h"
 
 #include "AutomationWorkerMessages.generated.h"
@@ -99,6 +100,10 @@ struct FAutomationWorkerFindWorkersResponse
 	UPROPERTY(EditAnywhere, Category="Message")
 	FGuid SessionId;
 
+	/** Holds the name of the current RHI. */
+	UPROPERTY(EditAnywhere, Category = "Message")
+	FString RHIName;
+
 	/** Default constructor. */
 	FAutomationWorkerFindWorkersResponse() : RAMInGB(0) { }
 };
@@ -287,6 +292,10 @@ struct FAutomationWorkerRunTests
 	UPROPERTY()
 	FString BeautifiedTestName;
 
+	/** Holds the full test path of the test to run. */
+	UPROPERTY()
+	FString FullTestPath;
+
 	/** If true, send results to analytics when complete */
 	UPROPERTY()
 	bool bSendAnalytics;
@@ -295,11 +304,12 @@ struct FAutomationWorkerRunTests
 	FAutomationWorkerRunTests( ) :ExecutionCount(0), RoleIndex(0), bSendAnalytics(false) { }
 
 	/** Creates and initializes a new instance. */
-	FAutomationWorkerRunTests( uint32 InExecutionCount, int32 InRoleIndex, FString InTestName, FString InBeautifiedTestName, bool InSendAnalytics)
+	FAutomationWorkerRunTests( uint32 InExecutionCount, int32 InRoleIndex, FString InTestName, FString InBeautifiedTestName, FString InFullTestPath, bool InSendAnalytics)
 		: ExecutionCount(InExecutionCount)
 		, RoleIndex(InRoleIndex)
 		, TestName(InTestName)
 		, BeautifiedTestName(InBeautifiedTestName)
+		, FullTestPath(InFullTestPath)
 		, bSendAnalytics(InSendAnalytics)
 	{ }
 };
@@ -339,7 +349,7 @@ public:
 
 	/** */
 	UPROPERTY(EditAnywhere, Category="Message")
-	bool Success = false;
+	EAutomationState State = EAutomationState::NotRun;
 };
 
 
diff --git a/Engine/Source/Runtime/AutomationTest/AutomationTest.Build.cs b/Engine/Source/Runtime/AutomationTest/AutomationTest.Build.cs
new file mode 100644
--- /dev/null
+++ b/Engine/Source/Runtime/AutomationTest/AutomationTest.Build.cs
@@ -0,0 +1,22 @@
+// Copyright Epic Games, Inc. All Rights Reserved.
+
+namespace UnrealBuildTool.Rules
+{
+	public class AutomationTest : ModuleRules
+	{
+		public AutomationTest(ReadOnlyTargetRules Target) : base(Target)
+		{
+			PublicDependencyModuleNames.AddRange(
+				new string[]
+				{
+					"CoreUObject",
+				});
+
+			PrivateDependencyModuleNames.AddRange(
+				new string[]
+				{
+					"Core",
+				});
+		}
+	}
+}
diff --git a/Engine/Source/Runtime/AutomationTest/Private/AutomationTestExcludelist.cpp b/Engine/Source/Runtime/AutomationTest/Private/AutomationTestExcludelist.cpp
new file mode 100644
--- /dev/null
+++ b/Engine/Source/Runtime/AutomationTest/Private/AutomationTestExcludelist.cpp
@@ -0,0 +1,123 @@
+// Copyright Epic Games, Inc. All Rights Reserved.
+
+#include "AutomationTestExcludelist.h"
+
+static const FString FunctionalTestsPreFix = TEXT("Project.Functional Tests.");
+
+void UAutomationTestExcludelist::OverrideConfigSection(FString& SectionName)
+{
+	SectionName = TEXT("AutomationTestExcludelist");
+}
+
+UAutomationTestExcludelist* UAutomationTestExcludelist::Get()
+{
+	return GetMutableDefault<UAutomationTestExcludelist>();
+}
+
+void UAutomationTestExcludelist::PostInitProperties()
+{
+	Super::PostInitProperties();
+
+	for (auto& Entry : ExcludeTest)
+	{
+		if (Entry.IsEmpty())
+		{
+			Entry.FullTestName = GetFullTestName(Entry);
+		}
+	}
+}
+
+FString UAutomationTestExcludelist::GetFullTestName(const FAutomationTestExcludelistEntry& ExcludelistEntry)
+{
+	if (!ExcludelistEntry.IsEmpty())
+	{
+		return ExcludelistEntry.FullTestName;
+	}
+
+	FString Map = ExcludelistEntry.Map.ToString().TrimStartAndEnd();
+	FString ListName = ExcludelistEntry.Test.ToString().TrimStartAndEnd();
+
+	if (Map.StartsWith(TEXT("/")) && !ListName.StartsWith(FunctionalTestsPreFix))
+	{
+		ListName = FunctionalTestsPreFix + Map + TEXT(".") + ListName;
+	}
+
+	return ListName.ToLower();
+}
+
+void UAutomationTestExcludelist::AddToExcludeTest(const FString& TestName, const FAutomationTestExcludelistEntry& ExcludelistEntry)
+{
+	auto NewEntry = FAutomationTestExcludelistEntry(ExcludelistEntry);
+	NewEntry.Test = *TestName;
+	if (!NewEntry.Map.IsNone())
+	{
+		NewEntry.Map = TEXT("");
+	}
+
+	NewEntry.FullTestName = GetFullTestName(NewEntry);
+
+	ExcludeTest.Add(NewEntry);
+}
+
+void UAutomationTestExcludelist::RemoveFromExcludeTest(const FString& TestName)
+{
+	if (TestName.IsEmpty())
+		return;
+
+	const FString NameToCompare = TestName.ToLower();
+
+	for (int i = 0; i < ExcludeTest.Num(); ++i)
+	{
+		if (ExcludeTest[i].FullTestName == NameToCompare)
+		{
+			ExcludeTest.RemoveAt(i);
+			return;
+		}
+	}
+}
+
+bool UAutomationTestExcludelist::IsTestExcluded(const FString& TestName, const FString& RHI, FName* OutReason, bool* OutWarn)
+{
+	if (auto Entry = GetExcludeTestEntry(TestName))
+	{
+		if (!Entry->ShouldExcludeForRHI(RHI))
+		{
+			return false;
+		}
+
+		if (OutReason != nullptr)
+		{
+			*OutReason = Entry->Reason;
+		}
+
+		if (OutWarn != nullptr)
+		{
+			*OutWarn = Entry->Warn;
+		}
+
+		return true;
+	}
+
+	return false;
+}
+
+FAutomationTestExcludelistEntry* UAutomationTestExcludelist::GetExcludeTestEntry(const FString& TestName)
+{
+	if (TestName.IsEmpty())
+		return nullptr;
+
+	const FString NameToCompare = TestName.ToLower();
+
+	for (auto& Entry : ExcludeTest)
+	{
+		if(NameToCompare.StartsWith(Entry.FullTestName))
+		{
+			if (NameToCompare.Len() == Entry.FullTestName.Len() || NameToCompare.Mid(Entry.FullTestName.Len(), 1) == TEXT("."))
+			{
+				return &Entry;
+			}
+		}
+	}
+
+	return nullptr;
+}
diff --git a/Engine/Source/Runtime/AutomationTest/Private/AutomationTestModule.cpp b/Engine/Source/Runtime/AutomationTest/Private/AutomationTestModule.cpp
new file mode 100644
--- /dev/null
+++ b/Engine/Source/Runtime/AutomationTest/Private/AutomationTestModule.cpp
@@ -0,0 +1,30 @@
+// Copyright Epic Games, Inc. All Rights Reserved.
+#include "AutomationTestModule.h"
+
+#include "Modules/ModuleInterface.h"
+#include "Modules/ModuleManager.h"
+
+
+DEFINE_LOG_CATEGORY(LogAutomationTest)
+
+/**
+ * Implements the AutomationTesting module.
+ */
+class FAutomationTestModule
+	: public IModuleInterface
+{
+public:
+
+	//~ IModuleInterface interface
+
+	virtual void StartupModule() override { }
+	virtual void ShutdownModule() override { }
+
+	virtual bool SupportsDynamicReloading() override
+	{
+		return true;
+	}
+};
+
+
+IMPLEMENT_MODULE(FAutomationTestModule, AutomationTest);
diff --git a/Engine/Source/Runtime/AutomationTest/Public/AutomationState.h b/Engine/Source/Runtime/AutomationTest/Public/AutomationState.h
new file mode 100644
--- /dev/null
+++ b/Engine/Source/Runtime/AutomationTest/Public/AutomationState.h
@@ -0,0 +1,26 @@
+// Copyright Epic Games, Inc. All Rights Reserved.
+
+#pragma once
+
+#include "CoreTypes.h"
+#include "UObject/ObjectMacros.h"
+
+#include "AutomationState.generated.h"
+
+
+/** Enumeration of unit test status for special dialog */
+UENUM()
+enum class EAutomationState : uint8
+{
+	NotRun,					// Automation test was not run
+	InProcess,				// Automation test is running now
+	Fail,					// Automation test was run and failed
+	Success,				// Automation test was run and succeeded
+	Skipped,				// Automation test was skipped
+};
+
+
+inline const FString AutomationStateToString(EAutomationState InValue)
+{
+	return UEnum::GetDisplayValueAsText(InValue).ToString();
+}
diff --git a/Engine/Source/Runtime/AutomationTest/Public/AutomationTestExcludelist.h b/Engine/Source/Runtime/AutomationTest/Public/AutomationTestExcludelist.h
new file mode 100644
--- /dev/null
+++ b/Engine/Source/Runtime/AutomationTest/Public/AutomationTestExcludelist.h
@@ -0,0 +1,186 @@
+// Copyright Epic Games, Inc. All Rights Reserved.
+
+#pragma once
+
+#include "CoreMinimal.h"
+#include "AutomationTestExcludelist.generated.h"
+
+UENUM(Meta = (Bitflags, UseEnumValuesAsMaskValuesInEditor = "true"))
+enum class ERHI_Flags : uint8
+{
+	DirectX11	= 0x01 UMETA(DisplayName = "DirectX 11"),
+	DirectX12	= 0x02 UMETA(DisplayName = "DirectX 12"),
+	Vulkan		= 0x04 UMETA(DisplayName = "Vulkan"),
+	NUM			UMETA(Hidden)
+};
+
+ENUM_CLASS_FLAGS(ERHI_Flags);
+
+USTRUCT()
+struct FAutomationTestExcludeOptions
+{
+	GENERATED_BODY()
+
+	void SetRHIFlagsFromNames(TSet<FName>& InRHIs)
+	{
+		int32 Num_RHIs = InRHIs.Num();
+		if (Num_RHIs == 0)
+			return;
+
+		int32 InRHIs_index = 0;
+		UEnum* Enum = StaticEnum<ERHI_Flags>();
+		int32 Num_Flags = Enum->NumEnums();
+		for (int32 i = 0; i < Num_Flags && InRHIs_index < Num_RHIs; i++)
+		{
+			if (InRHIs.Contains(*Enum->GetDisplayNameTextByIndex(i).ToString()))
+			{
+				RHIs |= (int8)Enum->GetValueByIndex(i);
+				InRHIs_index++;
+			}
+		}
+	}
+	/* Name of the target test */
+	UPROPERTY(VisibleAnywhere, Category = ExcludeTestOptions)
+	FName Test;
+
+	/* Reason to why the test is excluded */
+	UPROPERTY(EditAnywhere, Category = ExcludeTestOptions)
+	FString Reason;
+
+	/* Option to target specific RHI. No option means it should be applied to all RHI */
+	UPROPERTY(EditAnywhere, Meta = (Bitmask, BitmaskEnum = "ERHI_Flags"), Category = ExcludeTestOptions)
+	int8 RHIs = 0;
+
+	/* Should the Reason be reported as a warning in the log */
+	UPROPERTY(EditAnywhere, Category = ExcludeTestOptions)
+	bool Warn = false;
+};
+
+USTRUCT()
+struct FAutomationTestExcludelistEntry
+{
+	GENERATED_BODY()
+
+	FAutomationTestExcludelistEntry() { }
+
+	FAutomationTestExcludelistEntry(const FAutomationTestExcludeOptions& Options)
+		: Test(Options.Test)
+		, Reason(Options.Reason)
+		, Warn(Options.Warn)
+	{
+		UEnum* Enum = StaticEnum<ERHI_Flags>();
+		int32 Num_Flags = Enum->NumEnums();
+		for (int32 i = 0; i < Num_Flags; i++)
+		{
+			int8 Flag = Enum->GetValueByIndex(i);
+			if ((int8)ERHI_Flags::NUM == Flag)
+				break; // We reach the maximum value
+
+			if ((Options.RHIs & Flag) == Flag)
+			{
+				RHIs.Add(*Enum->GetDisplayNameTextByIndex(i).ToString());
+			}
+		}
+	}
+
+	TSharedPtr<FAutomationTestExcludeOptions> GetOptions()
+	{
+		TSharedPtr<FAutomationTestExcludeOptions> Options = MakeShareable(new FAutomationTestExcludeOptions());
+		Options->Test = Test;
+		Options->Reason = Reason.ToString();
+		Options->Warn = Warn;
+		Options->SetRHIFlagsFromNames(RHIs);
+		
+		return Options;
+	}
+
+	/* Determine if exclusion entry is propagated based on test name - used for management in test automation window */
+	void SetPropagation(const FString& ForTestName)
+	{
+		bIsPropagated = FullTestName != ForTestName.ToLower();
+	}
+
+	/* Return true if the entry is not specific */
+	bool IsEmpty() const
+	{
+		return FullTestName.IsEmpty();
+	}
+
+	/* Reset entry to be un-specific */
+	void Reset()
+	{
+		FullTestName.Empty();
+		bIsPropagated = false;
+	}
+
+	/* Confirm that test should be excluded for that RHI */
+	bool ShouldExcludeForRHI(const FString& InRHI) const
+	{
+		if (!InRHI.IsEmpty() && RHIs.Num() != 0 && !RHIs.Contains(*InRHI))
+		{
+			return false;
+		}
+
+		return true;
+	}
+
+	/* Hold full test name/path */
+	FString FullTestName;
+	/* Is the entry comes from propagation */
+	bool bIsPropagated = false;
+
+	// Use FName instead of FString to read params from config that aren't wrapped with quotes
+
+	/* Hold the name of the target functional test map */
+	UPROPERTY(EditDefaultsOnly, Category = AutomationTestExcludelist)
+	FName Map;
+
+	/* Hold the name of the target test - full test name is require here unless for functional tests */
+	UPROPERTY(EditDefaultsOnly, Category = AutomationTestExcludelist)
+	FName Test;
+
+	/* Reason to why the test is excluded */
+	UPROPERTY(EditDefaultsOnly, Category = AutomationTestExcludelist)
+	FName Reason;
+
+	/* Option to target specific RHI. Empty array means it should be applied to all RHI */
+	UPROPERTY(EditDefaultsOnly, Category = AutomationTestExcludelist)
+	TSet<FName> RHIs;
+
+	/* Should the Reason be reported as a warning in the log */
+	UPROPERTY(EditDefaultsOnly, Category = AutomationTestExcludelist)
+	bool Warn = false;
+};
+
+
+UCLASS(config = Engine, defaultconfig)
+class AUTOMATIONTEST_API UAutomationTestExcludelist : public UObject
+{
+	GENERATED_BODY()
+
+public:
+	UAutomationTestExcludelist() {}
+
+	static UAutomationTestExcludelist* Get();
+
+	void AddToExcludeTest(const FString& TestName, const FAutomationTestExcludelistEntry& ExcludelistEntry);
+	void RemoveFromExcludeTest(const FString& TestName);
+	bool IsTestExcluded(const FString& TestName, const FString& RHI = TEXT(""), FName* OutReason = nullptr, bool* OutWarn = nullptr);
+	FAutomationTestExcludelistEntry* GetExcludeTestEntry(const FString& TestName);
+
+	void SaveConfig() { UObject::UpdateDefaultConfigFile(); }
+	// It is called automatically when CDO is created, usually you don't need to call LoadConfig manually
+	void LoadConfig() { UObject::LoadConfig(GetClass()); }
+
+	virtual void OverrideConfigSection(FString& SectionName) override;
+
+protected:
+	virtual void PostInitProperties() override;
+
+private:
+	FString GetFullTestName(const FAutomationTestExcludelistEntry& ExcludelistEntry);
+
+	UPROPERTY(EditDefaultsOnly, globalconfig, Category = AutomationTestExcludelist)
+	TArray<FAutomationTestExcludelistEntry> ExcludeTest;
+};
+
diff --git a/Engine/Source/Runtime/AutomationTest/Public/AutomationTestModule.h b/Engine/Source/Runtime/AutomationTest/Public/AutomationTestModule.h
new file mode 100644
--- /dev/null
+++ b/Engine/Source/Runtime/AutomationTest/Public/AutomationTestModule.h
@@ -0,0 +1,8 @@
+// Copyright Epic Games, Inc. All Rights Reserved.
+#pragma once
+
+#include "CoreTypes.h"
+
+#include "Logging/LogMacros.h"
+
+AUTOMATIONTEST_API DECLARE_LOG_CATEGORY_EXTERN(LogAutomationTest, Log, All);
\ No newline at end of file
diff --git a/Engine/Source/Runtime/AutomationWorker/AutomationWorker.Build.cs b/Engine/Source/Runtime/AutomationWorker/AutomationWorker.Build.cs
--- a/Engine/Source/Runtime/AutomationWorker/AutomationWorker.Build.cs
+++ b/Engine/Source/Runtime/AutomationWorker/AutomationWorker.Build.cs
@@ -17,6 +17,7 @@ namespace UnrealBuildTool.Rules
 				new string[]
 				{
 					"AutomationMessages",
+					"AutomationTest",
 					"CoreUObject",
                     "Analytics",
     				"AnalyticsET",
diff --git a/Engine/Source/Runtime/AutomationWorker/Private/AutomationAnalytics.cpp b/Engine/Source/Runtime/AutomationWorker/Private/AutomationAnalytics.cpp
--- a/Engine/Source/Runtime/AutomationWorker/Private/AutomationAnalytics.cpp
+++ b/Engine/Source/Runtime/AutomationWorker/Private/AutomationAnalytics.cpp
@@ -233,7 +233,7 @@ void FAutomationAnalytics::FireEvent_AutomationTestResults(const FAutomationWork
 		ParamArray.Add(FAnalyticsEventAttribute(GetAutomationParamName(EAutomationAnalyticParam::TestName), TestResults->TestName));
 		ParamArray.Add(FAnalyticsEventAttribute(GetAutomationParamName(EAutomationAnalyticParam::BeautifiedName), BeautifiedTestName));
 		ParamArray.Add(FAnalyticsEventAttribute(GetAutomationParamName(EAutomationAnalyticParam::ExecutionCount), TestResults->ExecutionCount));
-		ParamArray.Add(FAnalyticsEventAttribute(GetAutomationParamName(EAutomationAnalyticParam::IsSuccess), TestResults->Success));
+		ParamArray.Add(FAnalyticsEventAttribute(GetAutomationParamName(EAutomationAnalyticParam::IsSuccess), TestResults->State == EAutomationState::Success));
 		ParamArray.Add(FAnalyticsEventAttribute(GetAutomationParamName(EAutomationAnalyticParam::Duration), TestResults->Duration));
 		ParamArray.Add(FAnalyticsEventAttribute(GetAutomationParamName(EAutomationAnalyticParam::ErrorCount), TestResults->ErrorTotal));
 		ParamArray.Add(FAnalyticsEventAttribute(GetAutomationParamName(EAutomationAnalyticParam::WarningCount), TestResults->WarningTotal));
diff --git a/Engine/Source/Runtime/AutomationWorker/Private/AutomationWorkerModule.cpp b/Engine/Source/Runtime/AutomationWorker/Private/AutomationWorkerModule.cpp
--- a/Engine/Source/Runtime/AutomationWorker/Private/AutomationWorkerModule.cpp
+++ b/Engine/Source/Runtime/AutomationWorker/Private/AutomationWorkerModule.cpp
@@ -4,6 +4,7 @@
 
 #include "AutomationAnalytics.h"
 #include "AutomationWorkerMessages.h"
+#include "AutomationTestExcludelist.h"
 #include "HAL/FileManager.h"
 #include "MessageEndpoint.h"
 #include "MessageEndpointBuilder.h"
@@ -42,8 +43,6 @@ void FAutomationWorkerModule::StartupModule()
 
 	FAutomationTestFramework::Get().PreTestingEvent.AddRaw(this, &FAutomationWorkerModule::HandlePreTestingEvent);
 	FAutomationTestFramework::Get().PostTestingEvent.AddRaw(this, &FAutomationWorkerModule::HandlePostTestingEvent);
-
-	FAutomationTestFramework::Get().BuildTestBlacklistFromConfig();
 }
 
 void FAutomationWorkerModule::ShutdownModule()
@@ -237,7 +236,7 @@ void FAutomationWorkerModule::ReportTestComplete()
 
 			Message->TestName = TestName;
 			Message->ExecutionCount = ExecutionCount;
-			Message->Success = bSuccess;
+			Message->State = bSuccess ? EAutomationState::Success : EAutomationState::Fail;
 			Message->Duration = ExecutionInfo.Duration;
 			Message->Entries = ExecutionInfo.GetEntries();
 			Message->WarningTotal = ExecutionInfo.GetWarningTotal();
@@ -339,6 +338,7 @@ void FAutomationWorkerModule::SendWorkerFound(const FMessageAddress& ControllerA
 	Response->GPUName = FPlatformMisc::GetPrimaryGPUBrand();
 	Response->CPUModelName = CPUModelString;
 	Response->RAMInGB = FPlatformMemory::GetPhysicalGBRam();
+	Response->RHIName = FApp::GetGraphicsRHI();
 #if WITH_ENGINE && WITH_AUTOMATION_TESTS
 	Response->RenderModeName = AutomationCommon::GetRenderDetailsString();
 #else
@@ -536,7 +536,7 @@ void FAutomationWorkerModule::HandleRunTestsMessage( const FAutomationWorkerRunT
 		FAutomationWorkerRunTestsReply* OutMessage = FMessageEndpoint::MakeMessage<FAutomationWorkerRunTestsReply>();
 		OutMessage->TestName = Message.TestName;
 		OutMessage->ExecutionCount = Message.ExecutionCount;
-		OutMessage->Success = false;
+		OutMessage->State = EAutomationState::Skipped;
 		OutMessage->Entries.Add(FAutomationExecutionEntry(FAutomationEvent(EAutomationEventType::Error, LogMessage)));
 		OutMessage->ErrorTotal = 1;
 		MessageEndpoint->Send(OutMessage, Context->GetSender());
@@ -544,6 +544,33 @@ void FAutomationWorkerModule::HandleRunTestsMessage( const FAutomationWorkerRunT
 		return;
 	}
 
+	// Do we need to skip the test
+	FName SkipReason;
+	bool bWarn(false);
+	UAutomationTestExcludelist* Excludelist = UAutomationTestExcludelist::Get();
+	if (Excludelist->IsTestExcluded(Message.FullTestPath, FApp::GetGraphicsRHI(), &SkipReason, &bWarn))
+	{
+		FString SkippingMessage = FString::Format(TEXT("Test Skipped. Name={{0}} Reason={{1}} Path={{2}}"),
+			{ *Message.BeautifiedTestName, *SkipReason.ToString(), *Message.FullTestPath });
+		if (bWarn)
+		{
+			UE_LOG(LogAutomationWorker, Warning, TEXT("%s"), *SkippingMessage);
+		}
+		else
+		{
+			UE_LOG(LogAutomationWorker, Display, TEXT("%s"), *SkippingMessage);
+		}
+
+		FAutomationWorkerRunTestsReply* OutMessage = FMessageEndpoint::MakeMessage<FAutomationWorkerRunTestsReply>();
+		OutMessage->TestName = Message.TestName;
+		OutMessage->ExecutionCount = Message.ExecutionCount;
+		OutMessage->State = EAutomationState::Skipped;
+		OutMessage->Entries.Add(FAutomationExecutionEntry(FAutomationEvent(EAutomationEventType::Info, FString::Printf(TEXT("Skipping test because of exclude list: %s"), *SkipReason.ToString()))));
+		MessageEndpoint->Send(OutMessage, Context->GetSender());
+
+		return;
+	}
+
 	ExecutionCount = Message.ExecutionCount;
 	TestName = Message.TestName;
 	BeautifiedTestName = Message.BeautifiedTestName;
diff --git a/Engine/Source/Runtime/Core/Private/Misc/AutomationTest.cpp b/Engine/Source/Runtime/Core/Private/Misc/AutomationTest.cpp
--- a/Engine/Source/Runtime/Core/Private/Misc/AutomationTest.cpp
+++ b/Engine/Source/Runtime/Core/Private/Misc/AutomationTest.cpp
@@ -512,121 +512,6 @@ void FAutomationTestFramework::LoadTestModules( )
 	}
 }
 
-void FAutomationTestFramework::BuildTestBlacklistFromConfig()
-{
-	TestBlacklist.Empty();
-	if (GConfig)
-	{
-
-		const FString CommandLine = FCommandLine::Get();
-
-		for (const FString& ConfigFilename : GConfig->GetFilenames())
-		{
-			FConfigSection* BlacklistSection = GConfig->GetSectionPrivate(TEXT("AutomationTestBlacklist"), false, true, ConfigFilename);
-			if (BlacklistSection)
-			{
-				// Parse all blacklist definitions of the format "BlacklistTest=(Map=/Game/Tests/MapName, Test=TestName, Reason="Foo")"
-				for (FConfigSection::TIterator Section(*BlacklistSection); Section; ++Section)
-				{
-					if (Section.Key() == TEXT("BlacklistTest"))
-					{
-						FString BlacklistValue = Section.Value().GetValue();
-						FString Map, Test, Reason, RHIs, Warn, ListName;
-						bool bSuccess = false;
-
-						if (FParse::Value(*BlacklistValue, TEXT("Test="), Test, true))
-						{
-							ListName = FString(Test);
-							FParse::Value(*BlacklistValue, TEXT("Map="), Map, true);
-							FParse::Value(*BlacklistValue, TEXT("Reason="), Reason);
-							FParse::Value(*BlacklistValue, TEXT("RHIs="), RHIs);
-							FParse::Value(*BlacklistValue, TEXT("Warn="), Warn);
-
-							if (Map.IsEmpty())
-							{
-								// Test with no Map property
-								bSuccess = true;
-							}
-							else if (Map.StartsWith(TEXT("/")))
-							{
-								// Account for Functional Tests based on Map - historically blacklisting was made only for functional tests
-								ListName = TEXT("Project.Functional Tests.") + Map + TEXT(".") + ListName;
-								bSuccess = true;
-							}
-
-						}
-
-						if (bSuccess)
-						{
-							if ((!Map.IsEmpty() && CommandLine.Contains(Map)) || CommandLine.Contains(Test))
-							{
-								UE_LOG(LogAutomationTest, Warning, TEXT("Test '%s' is blacklisted but allowing due to command line."), *BlacklistValue);
-							}
-							else
-							{
-								ListName.RemoveSpacesInline();
-								FBlacklistEntry& Entry = TestBlacklist.Add(ListName);
-								Entry.Map = Map;
-								Entry.Test = Test;
-								Entry.Reason = Reason;
-								if (!RHIs.IsEmpty())
-								{
-									RHIs.ToLower().ParseIntoArray(Entry.RHIs, TEXT(","), true);
-									for (FString& RHI : Entry.RHIs)
-									{
-										RHI.TrimStartAndEndInline();
-									}
-								}
-								Entry.bWarn = Warn.ToBool();
-							}
-						}
-						else
-						{
-							UE_LOG(LogAutomationTest, Error, TEXT("Invalid blacklisted test definition: '%s'"), *BlacklistValue);
-						}
-					}
-				}
-			}
-		}
-	}
-
-	if (TestBlacklist.Num() > 0)
-	{
-		UE_LOG(LogAutomationTest, Log, TEXT("Automated Test Blacklist:"));
-		for (auto& KV : TestBlacklist)
-		{
-			UE_LOG(LogAutomationTest, Log, TEXT("\tTest: %s"), *KV.Key);
-		}
-	}
-}
-
-bool FAutomationTestFramework::IsBlacklisted(const FString& TestName, FString* OutReason, bool *OutWarn) const
-{
-	const FString ListName = TestName.Replace(TEXT(" "), TEXT(""));
-	const FBlacklistEntry* Entry = TestBlacklist.Find(ListName);
-
-	if (Entry)
-	{
-		if (Entry->RHIs.Num() != 0 && !Entry->RHIs.Contains(FApp::GetGraphicsRHI().ToLower()))
-		{
-			return false;
-		}
-
-		if (OutReason != nullptr)
-		{
-			*OutReason = Entry->Reason;
-		}
-
-		if (OutWarn != nullptr)
-		{
-			*OutWarn = Entry->bWarn;
-		}
-	}
-
-	return Entry != nullptr;
-}
-
-
 void FAutomationTestFramework::GetValidTestNames( TArray<FAutomationTestInfo>& TestInfo ) const
 {
 	TestInfo.Empty();
@@ -694,28 +579,7 @@ void FAutomationTestFramework::GetValidTestNames( TArray<FAutomationTestInfo>& T
 		{
 			TArray<FAutomationTestInfo> TestsToAdd;
 			CurTest->GenerateTestNames(TestsToAdd);
-			for (FAutomationTestInfo& Test : TestsToAdd)
-			{
-				FString BlacklistReason;
-				bool bWarn(false);
-				FString TestName = Test.GetDisplayName();
-				if (!IsBlacklisted(TestName.Replace(TEXT(" "), TEXT("")), &BlacklistReason, &bWarn))
-				{
-					TestInfo.Add(MoveTemp(Test));
-				}
-				else
-				{
-					if (bWarn)
-					{
-						UE_LOG(LogAutomationTest, Warning, TEXT("Test '%s' is blacklisted. %s"), *TestName, *BlacklistReason);
-					}
-					else
-					{
-						UE_LOG(LogAutomationTest, Display, TEXT("Test '%s' is blacklisted. %s"), *TestName, *BlacklistReason);
-					}
-				}
-			}
-			
+			TestInfo.Append(TestsToAdd);			
 		}
 
 		// Make sure people are not writing complex tests that take forever to return the names of the tests
diff --git a/Engine/Source/Runtime/Core/Public/Misc/AutomationTest.h b/Engine/Source/Runtime/Core/Public/Misc/AutomationTest.h
--- a/Engine/Source/Runtime/Core/Public/Misc/AutomationTest.h
+++ b/Engine/Source/Runtime/Core/Public/Misc/AutomationTest.h
@@ -880,11 +880,6 @@ public:
 	 */
 	void LoadTestModules();
 
-	/**
-	 * Load the test Blacklist from the config.
-	 */
-	void BuildTestBlacklistFromConfig();
-
 	/**
 	 * Populates the provided array with the names of all tests in the framework that are valid to run for the current
 	 * application settings.
@@ -1051,19 +1046,6 @@ private:
 		 std::atomic<FFeedbackContext*> DestinationContext;
 	 };
 
-	//** Store information about blacklisted test */
-	struct FBlacklistEntry
-	{
-		FBlacklistEntry() :
-			bWarn(false) {}
-
-		FString Map;
-		FString Test;
-		FString Reason;
-		TArray<FString> RHIs;
-		bool bWarn;
-	};
-
 	friend class FAutomationTestOutputDevice;
 	/** Helper method called to prepare settings for automation testing to follow */
 	void PrepForAutomationTests();
@@ -1095,17 +1077,6 @@ private:
 	 */
 	bool InternalStopTest(FAutomationTestExecutionInfo& OutExecutionInfo);
 
-
-	/**
-	 * Internal helper method that verify if a test is black listed.
-	 *
-	 * @param	TestName		Beautified test name to be checked
-	 * @param	OutReason		Output the reason for the test being blacklisted
-	 * @param	OutWarn			Output true if the config ask for a warning message
-	 * @return	true if the TestName is part of the blacklist.
-	 */
-	bool IsBlacklisted(const FString& TestName, FString* OutReason = nullptr, bool* OutWarn = nullptr) const;
-
 	/** Constructor */
 	FAutomationTestFramework();
 
@@ -1164,8 +1135,6 @@ private:
 	bool bForceSmokeTests;
 
 	bool bCaptureStack;
-
-	TMap<FString, FBlacklistEntry> TestBlacklist;
 };
 
 
