PR #8106: Add SetFilter command to Automation framework (Contributed by [REDACTED_PERSON])

diff --git a/Engine/Source/Developer/AutomationController/Private/AutomationCommandline.cpp b/Engine/Source/Developer/AutomationController/Private/AutomationCommandline.cpp
--- a/Engine/Source/Developer/AutomationController/Private/AutomationCommandline.cpp
+++ b/Engine/Source/Developer/AutomationController/Private/AutomationCommandline.cpp
@@ -93,6 +93,7 @@ public:
 		FilterMaps.Add("Stress", EAutomationTestFlags::StressFilter);
 		FilterMaps.Add("Perf", EAutomationTestFlags::PerfFilter);
 		FilterMaps.Add("Product", EAutomationTestFlags::ProductFilter);
+		FilterMaps.Add("All", EAutomationTestFlags::FilterMask);
 	}
 	
 	void Shutdown()
@@ -624,6 +625,16 @@ public:
 					}
 					AutomationCommandQueue.Add(EAutomationCommand::RunFilter);
 				}
+				else if (FParse::Command(&TempCmd, TEXT("SetFilter")))
+				{
+					FlagToUse = TempCmd;
+					StringCommand = TempCmd;
+					if (FilterMaps.Contains(FlagToUse))
+					{
+						AutomationController->SetRequestedTestFlags(FilterMaps[FlagToUse]);
+						Ar.Logf(TEXT("Setting test filter: %s"), *FlagToUse);
+					}
+				}
 				else if (FParse::Command(&TempCmd, TEXT("RunAll")))
 				{
 					AutomationCommandQueue.Add(EAutomationCommand::RunAll);
@@ -642,6 +653,7 @@ public:
 					Ar.Logf(TEXT("\tAutomation RunTests <test string>"));
 					Ar.Logf(TEXT("\tAutomation RunAll "));
 					Ar.Logf(TEXT("\tAutomation RunFilter <filter name>"));
+					Ar.Logf(TEXT("\tAutomation SetFilter <filter name>"));
 					Ar.Logf(TEXT("\tAutomation Quit"));
 					bHandled = false;
 				}
