Make automated tests report more resilient in case of critical failure or timeout.
* Json report saved after each test execution
* Test artifacts are copied after each test execution
* Update html report to display `not run` and `in process` test state
* Remove undocumented command line option `checkpoint`

diff --git a/Engine/Source/Developer/AutomationController/Private/AutomationCommandline.cpp b/Engine/Source/Developer/AutomationController/Private/AutomationCommandline.cpp
--- a/Engine/Source/Developer/AutomationController/Private/AutomationCommandline.cpp
+++ b/Engine/Source/Developer/AutomationController/Private/AutomationCommandline.cpp
@@ -33,7 +33,6 @@ enum class EAutomationCommand : uint8
 {
 	ListAllTests,			//List all tests for the session
 	RunCommandLineTests,	//Run only tests that are listed on the commandline
-	RunCheckpointTests,		// Run only tests listed on the commandline with checkpoints in case of a crash.
 	RunAll,					//Run all the tests that are supported
 	RunFilter,              //
 	Quit					//quit the app when tests are done
@@ -266,24 +265,6 @@ if (bMeetsMatch)
 			}
 		}
 
-		// If we have the TestsRun array set up and are using the same command as before, clear out already run tests. 
-		if (TestsRun.Num() > 0)
-		{
-			if (TestsRun[0] == StringCommand)
-			{
-				for (int i = 1; i < TestsRun.Num(); i++)
-				{
-					if (OutTestNames.Remove(TestsRun[i]))
-					{
-						UE_LOG(LogAutomationCommandLine, Display, TEXT("Skipping %s due to Checkpoint."), *TestsRun[i]);
-					}
-				}
-			}
-			else
-			{
-				AutomationController->CleanUpCheckpointFile();
-			}
-		}
 	}
 
 	void FindWorkers(float DeltaTime)
@@ -394,29 +375,6 @@ if (bMeetsMatch)
 			AutomationController->OnTestsRefreshed().Remove(TestsRefreshedHandle);
 			TestsRefreshedHandle.Reset();
 		}
-		else if (AutomationCommand == EAutomationCommand::RunCheckpointTests)
-		{
-			TArray <FString> FilteredTestNames;
-			GenerateTestNamesFromCommandLine(AllTestNames, FilteredTestNames);
-			if (FilteredTestNames.Num())
-			{
-				AutomationController->StopTests();
-				AutomationController->SetEnabledTests(FilteredTestNames);
-				if (TestsRun.Num())
-				{
-					AutomationController->WriteLoadedCheckpointDataToFile();
-				}
-				else
-				{
-					AutomationController->WriteLineToCheckpointFile(StringCommand);
-				}
-				bRunTests = true;
-			}
-			else
-			{
-				AutomationTestState = EAutomationTestState::Complete;
-			}
-		}
 		else if (AutomationCommand == EAutomationCommand::RunFilter)
 		{
 			if (FilterMaps.Contains(StringCommand))
@@ -606,14 +564,6 @@ if (bMeetsMatch)
 					Ar.Logf(TEXT("Automation: RunTests='%s' Queued."), *StringCommand);
 					AutomationCommandQueue.Add(EAutomationCommand::RunCommandLineTests);
 				}
-				else if (FParse::Command(&TempCmd, TEXT("RunCheckpointedTests")))
-				{
-					StringCommand = TempCmd;
-					Ar.Logf(TEXT("Running all tests with checkpoints matching substring: %s"), *StringCommand);
-					AutomationCommandQueue.Add(EAutomationCommand::RunCheckpointTests);
-					TestsRun = AutomationController->GetCheckpointFileContents();
-					AutomationController->CleanUpCheckpointFile();
-				}
 				else if (FParse::Command(&TempCmd, TEXT("SetMinimumPriority")))
 				{
 					StringCommand = TempCmd;
@@ -761,10 +711,6 @@ private:
 	//Dictionary that maps flag names to flag values.
 	TMap<FString, int32> FilterMaps;
 
-	FString CheckpointCommand;
-
-	TArray<FString> TestsRun;
-
 	// Any that we encountered during processing. Used in 'Quit' to determine error code
 	TArray<FString> Errors;
 };
diff --git a/Engine/Source/Developer/AutomationController/Private/AutomationControllerManager.h b/Engine/Source/Developer/AutomationController/Private/AutomationControllerManager.h
--- a/Engine/Source/Developer/AutomationController/Private/AutomationControllerManager.h
+++ b/Engine/Source/Developer/AutomationController/Private/AutomationControllerManager.h
@@ -89,8 +89,10 @@ public:
 		, SucceededWithWarnings(0)
 		, Failed(0)
 		, NotRun(0)
+		, InProcess(0)
 		, TotalDuration(0)
 		, ComparisonExported(false)
+		, IsRequired(false)
 	{
 		if (FEngineVersion::Current().HasChangelist())
 		{
@@ -126,6 +128,9 @@ public:
 	UPROPERTY()
 	int32 NotRun;
 
+	UPROPERTY()
+	int32 InProcess;
+
 	UPROPERTY()
 	float TotalDuration;
 
@@ -138,9 +143,66 @@ public:
 	UPROPERTY()
 	TArray<FAutomatedTestResult> Tests;
 
+	TMap<FString, uint32> TestsMapIndex;
+
+	bool IsRequired;
+
 	int32 GetTotalTests() const
 	{
-		return Succeeded + SucceededWithWarnings + Failed + NotRun;
+		return Succeeded + SucceededWithWarnings + Failed + NotRun + InProcess;
+	}
+
+	void AddTestResult(const IAutomationReportPtr& TestReport)
+	{
+		FAutomatedTestResult TestResult;
+		TestResult.Test = TestReport;
+		TestResult.TestDisplayName = TestReport->GetDisplayName();
+		TestResult.FullTestPath = TestReport->GetFullTestPath();
+
+		TestsMapIndex.Add(TestReport->GetFullTestPath(), Tests.Num());
+		Tests.Add(TestResult);
+		NotRun++;
+	}
+
+	FAutomatedTestResult& GetTestResult(const IAutomationReportPtr& TestReport)
+	{
+		const FString& FullTestPath = TestReport->GetFullTestPath();
+		check(TestsMapIndex.Contains(FullTestPath));
+		return Tests[TestsMapIndex[FullTestPath]];
+	}
+
+	void UpdateTestResultStatus(const IAutomationReportPtr TestReport, EAutomationState State, bool bHasWarning = false)
+	{
+		FAutomatedTestResult& TestResult = GetTestResult(TestReport);
+		TestResult.State = State;
+
+		// Book keeping
+		switch (State)
+		{
+		case EAutomationState::Success:
+			if (bHasWarning)
+			{
+				SucceededWithWarnings++;
+			}
+			else
+			{
+				Succeeded++;
+			}
+			InProcess--;
+			break;
+		case EAutomationState::Fail:
+			Failed++;
+			InProcess--;
+			break;
+		case EAutomationState::InProcess:
+			NotRun--;
+			InProcess++;
+			break;
+		default:
+			NotRun++;
+			InProcess--;
+			break;
+		}
 	}
 
 	void ClearAllEntries()
@@ -149,7 +211,10 @@ public:
 		SucceededWithWarnings = 0;
 		Failed = 0;
 		NotRun = 0;
+		InProcess = 0;
 		TotalDuration = 0;
+		IsRequired = false;
+		TestsMapIndex.Empty();
 		Tests.Empty();
 	}
 };
@@ -331,17 +396,6 @@ public:
 
 	virtual FString GetReportOutputPath() const override;
 
-	// Checkpoint logic
-	virtual TArray<FString> GetCheckpointFileContents() override;
-
-	virtual FArchive* GetCheckpointFileForWrite() override;
-
-	virtual void CleanUpCheckpointFile() override;
-
-	virtual void WriteLoadedCheckpointDataToFile() override;
-
-	virtual void WriteLineToCheckpointFile(FString LineToWrite) override;
-
 	virtual void ResetAutomationTestTimeout(const TCHAR* Reason) override;
 	
 protected:
@@ -361,12 +415,12 @@ protected:
 	/**
 	 * Create a json file that contains all of our test report data at /saved/automation/logs/AutomationReport-{CL}-{DateTime}.json
 	 */
-	bool GenerateJsonTestPassSummary(const FAutomatedTestPassResults& SerializedPassResults, FDateTime Timestamp);
+	bool GenerateJsonTestPassSummary(FAutomatedTestPassResults& SerializedPassResults);
 
 	/**
 	 * Generates a full html report of the testing, which may include links to images.  All of it will be bundled under a folder.
 	 */
-	bool GenerateHtmlTestPassSummary(const FAutomatedTestPassResults& SerializedPassResults, FDateTime Timestamp);
+	bool GenerateTestPassHtmlIndex();
 
 	/**
 	* Gather all info, warning, and error lines generated over the course of a test.
@@ -543,7 +597,7 @@ private:
 	bool bSendAnalytics = false;
 
 	/** The list of results generated by our test pass. */
-	FAutomatedTestPassResults OurPassResults;
+	FAutomatedTestPassResults JsonTestPassResults;
 
 	/** The screenshot manager. */
 	IScreenShotManagerPtr ScreenshotManager;
@@ -567,15 +621,6 @@ private:
 
 	FString DeveloperReportUrl;
 
-	// TODO Put into struct.
-	// Checkpoint variables
-	//Test pass checkpoint backup file.
-	FArchive* CheckpointFile;
-
-	FString CheckpointCommand;
-
-	TArray<FString> TestsRun;
-
 private:
 
 	/** Holds a delegate that is invoked when the controller shuts down. */
diff --git a/Engine/Source/Developer/AutomationController/Private/AutomationControllerManger.cpp b/Engine/Source/Developer/AutomationController/Private/AutomationControllerManger.cpp
--- a/Engine/Source/Developer/AutomationController/Private/AutomationControllerManger.cpp
+++ b/Engine/Source/Developer/AutomationController/Private/AutomationControllerManger.cpp
@@ -58,9 +58,6 @@ FAutomationControllerManager::FAutomationControllerManager()
 		GameInstanceLostTimerSeconds = Settings->GameInstanceLostTimerSeconds;
 	}
 	
-	CheckpointFile = nullptr;
-
-
 	FString DeveloperPath;
 	FParse::Value(FCommandLine::Get(), TEXT("ReportOutputPath="), ReportExportPath, false);
 	FParse::Value(FCommandLine::Get(), TEXT("DisplayReportOutputPath="), ReportURLPath, false);
@@ -176,6 +173,19 @@ void FAutomationControllerManager::RunTests(const bool bInIsLocalSession)
 	//reset all tests
 	ReportManager.ResetForExecution(NumTestPasses);
 
+	// Register All tests that we'll need to be exported as json report
+	if (!ReportExportPath.IsEmpty())
+	{
+		JsonTestPassResults.IsRequired = true;
+		TArray<IAutomationReportPtr> TestsToRun = ReportManager.GetEnabledTestReports();
+		for (IAutomationReportPtr& TestReport : TestsToRun)
+		{
+			JsonTestPassResults.AddTestResult(TestReport);
+		}
+		// Get the html index written down as soon as possible so once results are updated, they can be reviewed.
+		GenerateTestPassHtmlIndex();
+	}
+
 	for ( int32 ClusterIndex = 0; ClusterIndex < DeviceClusterManager.GetNumClusters(); ++ClusterIndex )
 	{
 		//enable each device cluster
@@ -378,54 +388,59 @@ void FAutomationControllerManager::ProcessAvailableTasks()
 void FAutomationControllerManager::ReportTestResults()
 {
 	UE_LOG(LogAutomationController, Log, TEXT("Test Pass Results:"));
-	for ( int32 i = 0; i < OurPassResults.Tests.Num(); i++ )
+	for ( int32 i = 0; i < JsonTestPassResults.Tests.Num(); i++ )
 	{
-		UE_LOG(LogAutomationController, Log, TEXT("%s: %s"), *OurPassResults.Tests[i].TestDisplayName, ToString(OurPassResults.Tests[i].State));
+		UE_LOG(LogAutomationController, Log, TEXT("%s: %s"), *JsonTestPassResults.Tests[i].TestDisplayName, ToString(JsonTestPassResults.Tests[i].State));
 	}
 }
 
 void FAutomationControllerManager::CollectTestResults(TSharedPtr<IAutomationReport> Report, const FAutomationTestResults& Results)
 {
-	// TODO This is slow, change to a map.
-	for ( int32 i = 0; i < OurPassResults.Tests.Num(); i++ )
+	if (JsonTestPassResults.IsRequired)
 	{
-		FAutomatedTestResult& ReportResult = OurPassResults.Tests[i];
-		if ( ReportResult.FullTestPath == Report->GetFullTestPath() )
+		JsonTestPassResults.UpdateTestResultStatus(Report, Results.State, Results.GetWarningTotal() > 0);
+		FAutomatedTestResult& TestResult = JsonTestPassResults.GetTestResult(Report);
+		TestResult.SetEvents(Results.GetEntries(), Results.GetWarningTotal(), Results.GetErrorTotal());
+		TestResult.SetArtifacts(Results.Artifacts);
+
+		JsonTestPassResults.TotalDuration += Results.Duration;
+
+		// Copy new artifacts to Report export path
+		FCriticalSection CS;
+		for (FAutomationArtifact& Artifact : TestResult.GetArtifacts())
 		{
-			ReportResult.State = Results.State;
-			ReportResult.SetEvents(Results.GetEntries(), Results.GetWarningTotal(), Results.GetErrorTotal());
-			ReportResult.SetArtifacts(Results.Artifacts);
+			TArray<FString> Keys;
+			Artifact.LocalFiles.GetKeys(Keys);
 
-			switch ( Results.State )
-			{
-			case EAutomationState::Success:
-				if ( Results.GetWarningTotal() > 0 )
-				{
-					OurPassResults.SucceededWithWarnings++;
-				}
-				else
+			ParallelFor(Keys.Num(), [&](int32 Index)
 				{
-					OurPassResults.Succeeded++;
-				}
-				break;
-			case EAutomationState::Fail:
-				OurPassResults.Failed++;
-				break;
-			default:
-				OurPassResults.NotRun++;
-				break;
-			}
-
-			OurPassResults.TotalDuration += Results.Duration;
+					const FString& Key = Keys[Index];
+					FString Path = CopyArtifact(ReportExportPath, Artifact.LocalFiles[Key]);
+					{
+						FScopeLock Lock(&CS);
+						Artifact.Files.Add(Key, MoveTemp(Path));
+					}
+					if (Key == TEXT("unapproved"))
+					{
+						// Copy screenshot report
+						FScreenshotExportResult ExportResult = ScreenshotManager->ExportScreenshotComparisonResult(Artifact.Name, ReportExportPath);
 
-			return;
+						FScopeLock Lock(&CS);
+						if (!JsonTestPassResults.ComparisonExported && ExportResult.Success)
+						{
+							JsonTestPassResults.ComparisonExported = ExportResult.Success;
+							JsonTestPassResults.ComparisonExportDirectory = ExportResult.ExportPath;
+						}
+					}
+				});
 		}
 	}
 }
 
-bool FAutomationControllerManager::GenerateJsonTestPassSummary(const FAutomatedTestPassResults& SerializedPassResults, FDateTime Timestamp)
+bool FAutomationControllerManager::GenerateJsonTestPassSummary(FAutomatedTestPassResults& SerializedPassResults)
 {
 	UE_LOG(LogAutomationController, Display, TEXT("Converting results to json object..."));
+	SerializedPassResults.ReportCreatedOn = FDateTime::Now();
 
 	FString Json;
 	if (FJsonObjectConverter::UStructToJsonObjectString(SerializedPassResults, Json))
@@ -455,7 +470,7 @@ bool FAutomationControllerManager::GenerateJsonTestPassSummary(const FAutomatedT
 	return false;
 }
 
-bool FAutomationControllerManager::GenerateHtmlTestPassSummary(const FAutomatedTestPassResults& SerializedPassResults, FDateTime Timestamp)
+bool FAutomationControllerManager::GenerateTestPassHtmlIndex()
 {
 	UE_LOG(LogAutomationController, Display, TEXT("Loading results html template..."));
 
@@ -549,14 +564,6 @@ void FAutomationControllerManager::ExecuteNextTask( int32 ClusterIndex, OUT bool
 					DeviceClusterManager.SetTest(ClusterIndex, DeviceIndex, NextTest);
 					TestsRunThisPass.Add(NextTest);
 
-					// Register this as a test we'll need to report on.
-					FAutomatedTestResult tempresult;
-					tempresult.Test = NextTest;
-					tempresult.TestDisplayName = NextTest->GetDisplayName();
-					tempresult.FullTestPath = NextTest->GetFullTestPath();
-
-					OurPassResults.Tests.Add(tempresult);
-
 					// If we now have enough devices reserved for the test, run it!
 					TArray<FMessageAddress> DeviceAddresses = DeviceClusterManager.GetDevicesReservedForTest(ClusterIndex, NextTest);
 					if ( DeviceAddresses.Num() == NextTest->GetNumParticipantsRequired() )
@@ -569,9 +576,11 @@ void FAutomationControllerManager::ExecuteNextTask( int32 ClusterIndex, OUT bool
 							UE_LOG(LogAutomationController, Display, AutomationTestStarting, *TestsRunThisPass[AddressIndex]->GetDisplayName(), *TestsRunThisPass[AddressIndex]->GetCommand());
 							TestResults.State = EAutomationState::InProcess;
 
-							if (CheckpointFile)
+							if (JsonTestPassResults.IsRequired)
 							{
-								WriteLineToCheckpointFile(NextTest->GetFullTestPath());
+								JsonTestPassResults.UpdateTestResultStatus(NextTest, EAutomationState::InProcess);
+								// Save the whole pass report so that if the next test triggers a critical failure we are not left with no pass report.
+								GenerateJsonTestPassSummary(JsonTestPassResults);
 							}
 
 							TestResults.GameInstance = DeviceClusterManager.GetClusterDeviceName(ClusterIndex, DeviceIndex);
@@ -734,36 +743,14 @@ void FAutomationControllerManager::ProcessResults()
 		}
 	}
 
-	if ( !ReportExportPath.IsEmpty() )
+	if ( JsonTestPassResults.IsRequired )
 	{
-		FDateTime StartTime = FDateTime::Now();
-
 		UE_LOG(LogAutomationController, Display, TEXT("Exporting Automation Report to %s."), *ReportExportPath);
 
-		FAutomatedTestPassResults SerializedPassResults;
-
-		{
-			FDateTime StepTime = FDateTime::Now();
-
-			UE_LOG(LogAutomationController, Log, TEXT("Exporting comparison results to %s..."), *ReportExportPath);
-
-			FScreenshotExportResults ExportResults = ScreenshotManager->ExportComparisonResultsAsync(ReportExportPath).Get();
-
-			SerializedPassResults = OurPassResults;
-
-			SerializedPassResults.ComparisonExported = ExportResults.Success;
-			SerializedPassResults.ComparisonExportDirectory = ExportResults.ExportPath;
-			SerializedPassResults.ReportCreatedOn = StartTime;
-
-			UE_LOG(LogAutomationController, Log, TEXT("Exported results in %.02f Seconds"), (FDateTime::Now() - StepTime).GetTotalSeconds());
-		}
-
-
-		{
-			FDateTime StepTime = FDateTime::Now();
-
-			UE_LOG(LogAutomationController, Log, TEXT("Copying artifacts to %s..."), *ReportExportPath);
+		FAutomatedTestPassResults SerializedPassResults = JsonTestPassResults;
 
+		{ 
+			// Sort result by failure to improve readability
 			SerializedPassResults.Tests.StableSort([](const FAutomatedTestResult& A, const FAutomatedTestResult& B) {
 				if (A.GetErrorTotal() > 0)
 				{
@@ -791,58 +778,6 @@ void FAutomationControllerManager::ProcessResults()
 
 				return A.FullTestPath < B.FullTestPath;
 			});
-
-			// used for reporting and sync during copies
-			FCriticalSection CS;
-			int TotalFileArtifacts = 0;
-			int CopiedFileArtifacts = 0;
-
-			// Get a total for reporting
-			for (FAutomatedTestResult& Test : SerializedPassResults.Tests)
-			{
-				for (FAutomationArtifact& Artifact : Test.GetArtifacts())
-				{
-					TotalFileArtifacts += Artifact.LocalFiles.Num();
-				}
-			}
-
-
-			// most tests have a single set of artifacts with three or more files, so we could optimize this further by gathering them all first...
-			for (FAutomatedTestResult& Test : SerializedPassResults.Tests)
-			{
-				int TestArtifactCount = 0;
-
-				for (FAutomationArtifact& Artifact : Test.GetArtifacts())
-				{
-					TArray<FString> Keys;
-					Artifact.LocalFiles.GetKeys(Keys);
-
-					ParallelFor(Keys.Num(), [&](int32 Index)
-					{
-						const FString& Key = Keys[Index];
-
-						FString Path = CopyArtifact(ReportExportPath, Artifact.LocalFiles[Key]);
-						{
-							FScopeLock Lock(&CS);
-							Artifact.Files.Add(Key, MoveTemp(Path));
-							CopiedFileArtifacts++;
-							TestArtifactCount++;
-
-							// Show occasional progress for larger result sets
-							if ((CopiedFileArtifacts % 50) == 0)
-							{
-								UE_LOG(LogAutomationController, Log, TEXT("Copied %d of %d files in %.02f Seconds"), CopiedFileArtifacts, TotalFileArtifacts, (FDateTime::Now() - StepTime).GetTotalSeconds());
-							}
-						}
-					});
-
-					//UE_LOG(LogAutomationController, Verbose, TEXT("Copied %d files from artifact %s"), Keys.Num(), *Artifact.Name);
-				}
-
-				//UE_LOG(LogAutomationController, Verbose, TEXT("Copied %d files from test %s"), TestArtifactCount, *Test.TestDisplayName);
-			}
-
-			UE_LOG(LogAutomationController, Log, TEXT("Copied %d files in %.02f Seconds"), TotalFileArtifacts, (FDateTime::Now() - StepTime).GetTotalSeconds());
 		}
 
 		{
@@ -850,10 +785,7 @@ void FAutomationControllerManager::ProcessResults()
 			UE_LOG(LogAutomationController, Log, TEXT("Writing reports to %s..."), *ReportExportPath);
 
 			// Generate Json
-			GenerateJsonTestPassSummary(SerializedPassResults, StartTime);
-
-			// Generate Html
-			GenerateHtmlTestPassSummary(SerializedPassResults, StartTime);
+			GenerateJsonTestPassSummary(SerializedPassResults);
 
 			UE_LOG(LogAutomationController, Log, TEXT("Exported report to '%s' in %.02f Seconds"), *ReportExportPath, (FDateTime::Now() - StepTime).GetTotalSeconds());
 		}
@@ -874,8 +806,7 @@ void FAutomationControllerManager::ProcessResults()
 	UE_LOG(LogAutomationController, Display, TEXT("Report can be opened in the editor at '%s'"), ReportExportPath.IsEmpty() ? *FPaths::ConvertRelativePathToFull(FPaths::AutomationReportsDir()) : *ReportExportPath);
 
 	// Then clean our array for the next pass.
-	OurPassResults.ClearAllEntries();
-	CleanUpCheckpointFile();
+	JsonTestPassResults.ClearAllEntries();
 
 	SetControllerStatus(EAutomationControllerModuleState::Ready);
 }
@@ -1139,7 +1070,7 @@ void FAutomationControllerManager::HandleReceivedScreenShot(const FAutomationWor
 	// compare the incoming image and throw it away afterwards (note - there will be a copy in the report)
 	TSharedRef<FComparisonEntry> Comparison = MakeShareable(new FComparisonEntry());
 	Comparison->Sender = Context->GetSender();
-	Comparison->TestName = Message.Metadata.TestName;
+	Comparison->TestName = Message.Metadata.Context / Message.Metadata.TestName;
 	Comparison->PendingComparison = ScreenshotManager->CompareScreenshotAsync(IncomingFileName, Message.Metadata, EScreenShotCompareOptions::DiscardImage);
 
 	ComparisonQueue.Enqueue(Comparison);
@@ -1374,10 +1305,10 @@ void FAutomationControllerManager::HandleRunTestsReplyMessage(const FAutomationW
 
 			const FAutomationTestResults& FinalResults = Report->GetResults(ClusterIndex, CurrentTestPass);
 
+			ReportAutomationResult(Report, ClusterIndex, CurrentTestPass);
+
 			// Gather all of the data relevant to this test for our json reporting.
 			CollectTestResults(Report, FinalResults);
-
-			ReportAutomationResult(Report, ClusterIndex, CurrentTestPass);
 		}
 
 		// Device is now good to go
@@ -1415,72 +1346,6 @@ void FAutomationControllerManager::UpdateDeviceGroups( )
 	ReportManager.ClustersUpdated(NumOfClusters);
 }
 
-TArray<FString> FAutomationControllerManager::GetCheckpointFileContents()
-{
-	TestsRun.Empty();
-	FString CheckpointFileName = FString::Printf(TEXT("%sautomationcheckpoint.log"), *FPaths::AutomationDir());
-	if (IFileManager::Get().FileExists(*CheckpointFileName))
-	{
-		FString FileData;
-		FFileHelper::LoadFileToString(FileData, *CheckpointFileName);
-		FileData.ParseIntoArrayLines(TestsRun);
-		for (int i = 0; i < TestsRun.Num(); i++)
-		{
-			UE_LOG(LogAutomationController, Log, TEXT("AutomationCheckpoint %s"), *TestsRun[i]);
-		}
-	}
-	return TestsRun;
-}
-
-FArchive* FAutomationControllerManager::GetCheckpointFileForWrite()
-{
-	if (!CheckpointFile)
-	{
-		FString CheckpointFileName = FString::Printf(TEXT("%sautomationcheckpoint.log"), *FPaths::AutomationDir());
-		CheckpointFile = IFileManager::Get().CreateFileWriter(*CheckpointFileName, 8);
-	}
-	return CheckpointFile;
-}
-
-void FAutomationControllerManager::CleanUpCheckpointFile()
-{
-	if (CheckpointFile)
-	{
-		CheckpointFile->Close();
-		CheckpointFile = nullptr;
-	}
-	FString CheckpointFileName = FString::Printf(TEXT("%sautomationcheckpoint.log"), *FPaths::AutomationDir());
-	if (IFileManager::Get().FileExists(*CheckpointFileName))
-	{
-		IFileManager::Get().Delete(*CheckpointFileName);
-	}
-}
-
-void FAutomationControllerManager::WriteLoadedCheckpointDataToFile()
-{
-	GetCheckpointFileForWrite();
-	if (CheckpointFile)
-	{
-		for (int i = 0; i < TestsRun.Num(); i++)
-		{
-			FString LineToWrite = FString::Printf(TEXT("%s\r\n"), *TestsRun[i]);
-			CheckpointFile->Serialize(TCHAR_TO_ANSI(*LineToWrite), LineToWrite.Len());
-			CheckpointFile->Flush();
-		}
-	}
-}
-
-void FAutomationControllerManager::WriteLineToCheckpointFile(FString StringToWrite)
-{
-	GetCheckpointFileForWrite();
-	if (CheckpointFile)
-	{
-		FString LineToWrite = FString::Printf(TEXT("%s\r\n"), *StringToWrite);
-		CheckpointFile->Serialize(TCHAR_TO_ANSI(*LineToWrite), LineToWrite.Len());
-		CheckpointFile->Flush();
-	}
-}
-
 void FAutomationControllerManager::ResetAutomationTestTimeout(const TCHAR* Reason)
 {
 	//GLog->Logf(ELogVerbosity::Display, TEXT("Resetting automation test timeout: %s"), Reason);
diff --git a/Engine/Source/Developer/AutomationController/Private/AutomationReport.cpp b/Engine/Source/Developer/AutomationController/Private/AutomationReport.cpp
--- a/Engine/Source/Developer/AutomationController/Private/AutomationReport.cpp
+++ b/Engine/Source/Developer/AutomationController/Private/AutomationReport.cpp
@@ -634,6 +634,27 @@ TSharedPtr<IAutomationReport> FAutomationReport::GetNextReportToExecute(bool& bO
 	return NextReport;
 }
 
+void FAutomationReport::GetEnabledTestReports(TArray<TSharedPtr<IAutomationReport>>& OutReports)
+{
+	//if this is not a leaf node
+	if (ChildReports.Num())
+	{
+		for (int32 ReportIndex = 0; ReportIndex < ChildReports.Num(); ++ReportIndex)
+		{
+			ChildReports[ReportIndex]->GetEnabledTestReports(OutReports);
+		}
+	}
+	else
+	{
+		if (IsEnabled())
+		{
+			//Found one to run
+			OutReports.Add(AsShared());
+		}
+	}
+}
+
+
 const bool FAutomationReport::HasErrors()
 {
 	bool bHasErrors = false;
diff --git a/Engine/Source/Developer/AutomationController/Private/AutomationReport.h b/Engine/Source/Developer/AutomationController/Private/AutomationReport.h
--- a/Engine/Source/Developer/AutomationController/Private/AutomationReport.h
+++ b/Engine/Source/Developer/AutomationController/Private/AutomationReport.h
@@ -57,6 +57,7 @@ public:
 	virtual FString GetGameInstanceName( const int32 ClusterIndex ) override;
 	virtual TSharedPtr<IAutomationReport> EnsureReportExists(FAutomationTestInfo& TestInfo, const int32 ClusterIndex, const int32 NumPasses) override;
 	virtual TSharedPtr<IAutomationReport> GetNextReportToExecute(bool& bOutAllTestsComplete, const int32 ClusterIndex, const int32 PassIndex, const int32 NumDevicesInCluster) override;
+	virtual void GetEnabledTestReports(TArray<TSharedPtr<IAutomationReport>>& OutReports) override;
 	virtual const bool HasErrors() override;
 	virtual const bool HasWarnings() override;
 	virtual const bool GetDurationRange(float& OutMinTime, float& OutMaxTime) override;
diff --git a/Engine/Source/Developer/AutomationController/Private/AutomationReportManager.cpp b/Engine/Source/Developer/AutomationController/Private/AutomationReportManager.cpp
--- a/Engine/Source/Developer/AutomationController/Private/AutomationReportManager.cpp
+++ b/Engine/Source/Developer/AutomationController/Private/AutomationReportManager.cpp
@@ -48,6 +48,12 @@ TSharedPtr<IAutomationReport> FAutomationReportManager::GetNextReportToExecute(b
 	return ReportToExecute;
 }
 
+TArray<TSharedPtr<IAutomationReport>> FAutomationReportManager::GetEnabledTestReports()
+{
+	TArray<TSharedPtr<IAutomationReport>> Reports;
+	ReportRoot->GetEnabledTestReports(Reports);
+	return Reports;
+}
 
 TSharedPtr<IAutomationReport> FAutomationReportManager::EnsureReportExists(FAutomationTestInfo& TestInfo, const int32 ClusterIndex, const int32 NumPasses)
 {
diff --git a/Engine/Source/Developer/AutomationController/Private/AutomationReportManager.h b/Engine/Source/Developer/AutomationController/Private/AutomationReportManager.h
--- a/Engine/Source/Developer/AutomationController/Private/AutomationReportManager.h
+++ b/Engine/Source/Developer/AutomationController/Private/AutomationReportManager.h
@@ -39,6 +39,13 @@ public:
 	 */
 	TSharedPtr<IAutomationReport> GetNextReportToExecute(bool& bOutAllTestsComplete, const int32 ClusterIndex, const int32 PassIndex, const int32 NumDevicesInCluster);
 
+	/**
+	 * Returns all the enabled test reports.
+	 *
+	 * @return an array of all reports
+	 */
+	TArray<TSharedPtr<IAutomationReport>> GetEnabledTestReports();
+
 	/**
 	 * Ensures the nested tree exists ("blueprint.test.all") would have 3 levels other than the root.
 	 *
diff --git a/Engine/Source/Developer/AutomationController/Public/IAutomationControllerManager.h b/Engine/Source/Developer/AutomationController/Public/IAutomationControllerManager.h
--- a/Engine/Source/Developer/AutomationController/Public/IAutomationControllerManager.h
+++ b/Engine/Source/Developer/AutomationController/Public/IAutomationControllerManager.h
@@ -369,16 +369,6 @@ public:
 
 	virtual FString GetReportOutputPath() const = 0;
 
-	virtual TArray<FString> GetCheckpointFileContents() = 0;
-
-	virtual FArchive* GetCheckpointFileForWrite() = 0;
-
-	virtual void CleanUpCheckpointFile() = 0;
-
-	virtual void WriteLoadedCheckpointDataToFile() = 0;
-
-	virtual void WriteLineToCheckpointFile(FString LineToWrite) = 0;
-
 	virtual void ResetAutomationTestTimeout(const TCHAR* Reason) = 0;
 	
 public:
diff --git a/Engine/Source/Developer/AutomationController/Public/IAutomationReport.h b/Engine/Source/Developer/AutomationController/Public/IAutomationReport.h
--- a/Engine/Source/Developer/AutomationController/Public/IAutomationReport.h
+++ b/Engine/Source/Developer/AutomationController/Public/IAutomationReport.h
@@ -444,6 +444,13 @@ public:
 	 */
 	virtual TSharedPtr<IAutomationReport> GetNextReportToExecute(bool& bOutAllTestsComplete, const int32 ClusterIndex, const int32 PassIndex, const int32 NumDevicesInCluster) = 0;
 
+	/**
+	 * Returns all the enabled test reports.
+	 *
+	 * @param an array to be filled with all reports
+	 */
+	virtual void GetEnabledTestReports(TArray<TSharedPtr<IAutomationReport>>& OutReports) = 0;
+
 	/**
 	 * Returns if there have been any errors in the test.
 	 *
diff --git a/Engine/Source/Developer/ScreenShotComparisonTools/Private/ScreenShotManager.cpp b/Engine/Source/Developer/ScreenShotComparisonTools/Private/ScreenShotManager.cpp
--- a/Engine/Source/Developer/ScreenShotComparisonTools/Private/ScreenShotManager.cpp
+++ b/Engine/Source/Developer/ScreenShotComparisonTools/Private/ScreenShotManager.cpp
@@ -2,6 +2,7 @@
 
 #include "ScreenShotManager.h"
 #include "AutomationWorkerMessages.h"
+#include "Async/ParallelFor.h"
 #include "Async/Async.h"
 #include "HAL/FileManager.h"
 #include "MessageEndpointBuilder.h"
@@ -484,30 +485,27 @@ FImageComparisonResult FScreenShotManager::CompareScreenshot(const FString& InUn
 	return ComparisonResult;
 }
 
-TFuture<FScreenshotExportResults> FScreenShotManager::ExportComparisonResultsAsync(FString ExportPath)
-{
-	return Async(EAsyncExecution::Thread, [=] () { return ExportComparisonResults(ExportPath); });
-}
 
-FScreenshotExportResults FScreenShotManager::ExportComparisonResults(FString RootExportFolder)
+FScreenshotExportResult FScreenShotManager::ExportScreenshotComparisonResult(FString ScreenshotName, FString RootExportFolder)
 {
 	FPaths::NormalizeDirectoryName(RootExportFolder);
 
-	if ( RootExportFolder.IsEmpty() )
+	if (RootExportFolder.IsEmpty())
 	{
 		RootExportFolder = GetDefaultExportDirectory();
 	}
 
-	FScreenshotExportResults Results;
+	FScreenshotExportResult Results;
 	Results.Success = false;
 	Results.ExportPath = RootExportFolder / FString::FromInt(FEngineVersion::Current().GetChangelist());
 
-	if ( !IFileManager::Get().MakeDirectory(*Results.ExportPath, /*Tree =*/true) )
+	FString Destination = Results.ExportPath / ScreenshotName;
+	if (!IFileManager::Get().MakeDirectory(*Destination, /*Tree =*/true))
 	{
 		return Results;
 	}
 
-	CopyDirectory(Results.ExportPath, ScreenshotResultsFolder);
+	CopyDirectory(Destination, ScreenshotResultsFolder / ScreenshotName);
 
 	Results.Success = true;
 	return Results;
@@ -570,13 +568,13 @@ void FScreenShotManager::CopyDirectory(const FString& DestDir, const FString& Sr
 	FPaths::NormalizeDirectoryName(NormalizedSrc);
 
 	IFileManager::Get().FindFilesRecursive(FilesToCopy, *SrcDir, TEXT("*"), /*Files=*/true, /*Directories=*/false);
-	for ( const FString& File : FilesToCopy )
-	{
-		const FString& SourceFilePath = File;
-		FString DestFilePath = FPaths::Combine(DestDir, SourceFilePath.RightChop(SrcDir.Len()));
 
-		IFileManager::Get().Copy(*DestFilePath, *File, true, true);
-	}
+	ParallelFor(FilesToCopy.Num(), [&](int32 Index)
+		{
+			const FString& SourceFilePath = FilesToCopy[Index];
+			FString DestFilePath = FPaths::Combine(DestDir, SourceFilePath.RightChop(SrcDir.Len()));
+			IFileManager::Get().Copy(*DestFilePath, *SourceFilePath, true, true);
+		});
 }
 
 void FScreenShotManager::BuildFallbackPlatformsListFromConfig(const UScreenShotComparisonSettings* InSettings)
diff --git a/Engine/Source/Developer/ScreenShotComparisonTools/Private/ScreenShotManager.h b/Engine/Source/Developer/ScreenShotComparisonTools/Private/ScreenShotManager.h
--- a/Engine/Source/Developer/ScreenShotComparisonTools/Private/ScreenShotManager.h
+++ b/Engine/Source/Developer/ScreenShotComparisonTools/Private/ScreenShotManager.h
@@ -47,7 +47,7 @@ public:
 
 	virtual TFuture<FImageComparisonResult> CompareScreenshotAsync(const FString& IncomingPath, const FAutomationScreenshotMetadata& MetaData, const EScreenShotCompareOptions Options) override;
 
-	virtual TFuture<FScreenshotExportResults> ExportComparisonResultsAsync(FString ExportPath = TEXT("")) override;
+	virtual FScreenshotExportResult ExportScreenshotComparisonResult(FString ScreenshotName, FString ExportPath = TEXT("")) override;
 
 	virtual bool OpenComparisonReports(FString ImportPath, TArray<FComparisonReport>& OutReports) override;
 
@@ -68,7 +68,7 @@ private:
 	FString GetDefaultExportDirectory() const;
 
 	FImageComparisonResult CompareScreenshot(const FString& IncomingPath, const FAutomationScreenshotMetadata& MetaData, const EScreenShotCompareOptions Options);
-	FScreenshotExportResults ExportComparisonResults(FString RootExportFolder);
+
 	void CopyDirectory(const FString& DestDir, const FString& SrcDir);
 
 	void BuildFallbackPlatformsListFromConfig(const UScreenShotComparisonSettings* ScreenShotSettings);
diff --git a/Engine/Source/Developer/ScreenShotComparisonTools/Public/Interfaces/IScreenShotManager.h b/Engine/Source/Developer/ScreenShotComparisonTools/Public/Interfaces/IScreenShotManager.h
--- a/Engine/Source/Developer/ScreenShotComparisonTools/Public/Interfaces/IScreenShotManager.h
+++ b/Engine/Source/Developer/ScreenShotComparisonTools/Public/Interfaces/IScreenShotManager.h
@@ -23,12 +23,12 @@ typedef TSharedPtr<class IScreenShotManager> IScreenShotManagerPtr;
  */
 typedef TSharedRef<class IScreenShotManager> IScreenShotManagerRef;
 
-struct FScreenshotExportResults
+struct FScreenshotExportResult
 {
 	bool Success;
 	FString ExportPath;
 
-	FScreenshotExportResults()
+	FScreenshotExportResult()
 		: Success(false)
 	{
 	}
@@ -66,9 +66,9 @@ public:
 	virtual TFuture<FImageComparisonResult> CompareScreenshotAsync(const FString& IncomingPath, const FAutomationScreenshotMetadata& MetaData, const EScreenShotCompareOptions Options) = 0;
 
 	/**
-	 * Exports rs screenshots to the export location specified
+	 * Exports target screenshot report to the export location specified
 	 */
-	virtual TFuture<FScreenshotExportResults> ExportComparisonResultsAsync(FString ExportPath = TEXT("")) = 0;
+	virtual FScreenshotExportResult ExportScreenshotComparisonResult(FString ScreenshotName, FString ExportPath = TEXT("")) = 0;
 
 	/**
 	 * Imports screenshot comparison data from a given path.
diff --git a/Engine/Source/Programs/AutomationTool/Gauntlet/Framework/Utils/Gauntlet.HordeReport.cs b/Engine/Source/Programs/AutomationTool/Gauntlet/Framework/Utils/Gauntlet.HordeReport.cs
--- a/Engine/Source/Programs/AutomationTool/Gauntlet/Framework/Utils/Gauntlet.HordeReport.cs
+++ b/Engine/Source/Programs/AutomationTool/Gauntlet/Framework/Utils/Gauntlet.HordeReport.cs
@@ -208,8 +208,19 @@ namespace Gauntlet
 				get { return TestDetailed.State; }
 				set { TestDetailed.State = value; }
 			}
+			public int Errors
+			{
+				get { return TestDetailed.Errors; }
+				set { TestDetailed.Errors = value; }
+			}
+			public int Warnings
+			{
+				get { return TestDetailed.Warnings; }
+				set { TestDetailed.Warnings = value; }
+			}
 			public string ArtifactName { get; set; }
 
+
 			private TestResultDetailed TestDetailed { get; set; }
 
 			/// <summary>
