Avoid duplicates of Screenshot comparison artifacts when producing automated test html/json report

diff --git a/Engine/Source/Developer/AutomationController/Private/AutomationControllerManger.cpp b/Engine/Source/Developer/AutomationController/Private/AutomationControllerManger.cpp
--- a/Engine/Source/Developer/AutomationController/Private/AutomationControllerManger.cpp
+++ b/Engine/Source/Developer/AutomationController/Private/AutomationControllerManger.cpp
@@ -442,20 +442,20 @@ void FAutomationControllerManager::ProcessComparisonQueue()
 				// Record the artifacts for the test.
 				TMap<FString, FString> LocalFiles;
 
-				FString ProjectDir = FPaths::ProjectDir();
+				FString ScreenshotResultsFolder = FPaths::AutomationReportsDir();
 	
-				// Paths in the result are relative to the project directory.	
-				LocalFiles.Add(TEXT("unapproved"), FPaths::Combine(ProjectDir, Result.IncomingFilePath));
+				// Paths in the result are relative to the automation report directory.	
+				LocalFiles.Add(TEXT("unapproved"), FPaths::Combine(ScreenshotResultsFolder, Result.ReportApprovedFilePath));
 
 				// unapproved should always be valid. but approved/difference may be empty if this is a new screenshot
-				if (Result.ApprovedFilePath.Len())
+				if (Result.ReportIncomingFilePath.Len())
 				{
-					LocalFiles.Add(TEXT("approved"), FPaths::Combine(ProjectDir, Result.ApprovedFilePath));
+					LocalFiles.Add(TEXT("approved"), FPaths::Combine(ScreenshotResultsFolder, Result.ReportIncomingFilePath));
 				}
 
-				if (Result.ComparisonFilePath.Len())
+				if (Result.ReportComparisonFilePath.Len())
 				{
-					LocalFiles.Add(TEXT("difference"), FPaths::Combine(ProjectDir, Result.ComparisonFilePath));
+					LocalFiles.Add(TEXT("difference"), FPaths::Combine(ScreenshotResultsFolder, Result.ReportComparisonFilePath));
 				}
 
 				Report->AddArtifact(ClusterIndex, CurrentTestPass, FAutomationArtifact(UniqueId, Entry->ScreenshotPath, EAutomationArtifactType::Comparison, LocalFiles));
@@ -537,6 +537,13 @@ void FAutomationControllerManager::CollectTestResults(TSharedPtr<IAutomationRepo
 
 		JsonTestPassResults.TotalDuration += Results.Duration;
 
+		FString ImageCompareDirName = TEXT("imageCompare");
+		if (FEngineVersion::Current().GetChangelist() != 0)
+		{
+			ImageCompareDirName = ImageCompareDirName / FString::FromInt(FEngineVersion::Current().GetChangelist());
+		}
+		FString ImageCompareExportPath = ReportExportPath / ImageCompareDirName;
+
 		// Copy new artifacts to Report export path
 		FCriticalSection CS;
 		for (FAutomationArtifact& Artifact : TestResult.GetArtifacts())
@@ -547,7 +554,9 @@ void FAutomationControllerManager::CollectTestResults(TSharedPtr<IAutomationRepo
 			ParallelFor(Keys.Num(), [&](int32 Index)
 				{
 					const FString& Key = Keys[Index];
-					FString Path = CopyArtifact(ReportExportPath, Artifact.LocalFiles[Key]);
+					FString Path = Artifact.LocalFiles[Key];
+					FPaths::MakePathRelativeTo(Path, *FPaths::AutomationReportsDir());
+					Path = ImageCompareDirName / Path;
 					{
 						FScopeLock Lock(&CS);
 						Artifact.Files.Add(Key, MoveTemp(Path));
@@ -555,7 +564,7 @@ void FAutomationControllerManager::CollectTestResults(TSharedPtr<IAutomationRepo
 					if (Key == TEXT("unapproved"))
 					{
 						// Copy screenshot report
-						FScreenshotExportResult ExportResult = ScreenshotManager->ExportScreenshotComparisonResult(Artifact.Name, ReportExportPath);
+						FScreenshotExportResult ExportResult = ScreenshotManager->ExportScreenshotComparisonResult(Artifact.Name, ImageCompareExportPath);
 
 						FScopeLock Lock(&CS);
 						if (!JsonTestPassResults.ComparisonExported && ExportResult.Success)
@@ -680,22 +689,6 @@ FString FAutomationControllerManager::SlugString(const FString& DisplayString) c
 	return GeneratedName;
 }
 
-FString FAutomationControllerManager::CopyArtifact(const FString& DestFolder, const FString& SourceFile) const
-{
-	FString ArtifactFile = FString(TEXT("artifacts")) / FGuid::NewGuid().ToString(EGuidFormats::Digits) + FPaths::GetExtension(SourceFile, true);
-	FString ArtifactDestination = DestFolder / ArtifactFile;
-	
-	if (IFileManager::Get().Copy(*ArtifactDestination, *SourceFile, true, true) != 0)
-	{
-		uint32 ErrorCode = FPlatformMisc::GetLastError();
-		TCHAR ErrorBuffer[1024];
-		FPlatformMisc::GetSystemErrorMessage(ErrorBuffer, 1024, ErrorCode);
-		UE_LOG(LogAutomationController, Error, TEXT("Failed to copy %s to %s. Error: %u (%s)"), *SourceFile, *ArtifactDestination, ErrorCode, ErrorBuffer);
-	}
-
-	return ArtifactFile;
-}
-
 FString FAutomationControllerManager::GetReportOutputPath() const
 {
 	return ReportExportPath;
diff --git a/Engine/Source/Developer/ScreenShotComparison/Private/Models/ScreenComparisonModel.cpp b/Engine/Source/Developer/ScreenShotComparison/Private/Models/ScreenComparisonModel.cpp
--- a/Engine/Source/Developer/ScreenShotComparison/Private/Models/ScreenComparisonModel.cpp
+++ b/Engine/Source/Developer/ScreenShotComparison/Private/Models/ScreenComparisonModel.cpp
@@ -22,7 +22,7 @@ FScreenComparisonModel::FScreenComparisonModel(const FComparisonReport& InReport
 		that is where a blessed image should be checked in to.
 	*/
 
-	FString SourceImage = FPaths::Combine(Report.GetReportPath(), Report.GetComparisonResult().ReportIncomingFilePath);
+	FString SourceImage = FPaths::Combine(Report.GetReportRootDirectory(), Report.GetComparisonResult().ReportIncomingFilePath);
 	FString OutputImage = FPaths::Combine(FPaths::ProjectDir(), Report.GetComparisonResult().IdealApprovedFolderPath, FPaths::GetCleanFilename(Report.GetComparisonResult().IncomingFilePath));
 
 	FileImports.Add(FFileMapping(OutputImage, SourceImage));
@@ -53,7 +53,7 @@ TOptional<FAutomationScreenshotMetadata> FScreenComparisonModel::GetMetadata()
 	{
 		const FImageComparisonResult& Comparison = Report.GetComparisonResult();
 
-		FString IncomingImage = Report.GetReportPath() / Comparison.ReportIncomingFilePath;
+		FString IncomingImage = Report.GetReportRootDirectory() / Comparison.ReportIncomingFilePath;
 		FString IncomingMetadata = FPaths::ChangeExtension(IncomingImage, TEXT("json"));
 
 		if ( !IncomingMetadata.IsEmpty() )
diff --git a/Engine/Source/Developer/ScreenShotComparison/Private/Widgets/SScreenComparisonRow.cpp b/Engine/Source/Developer/ScreenShotComparison/Private/Widgets/SScreenComparisonRow.cpp
--- a/Engine/Source/Developer/ScreenShotComparison/Private/Widgets/SScreenComparisonRow.cpp
+++ b/Engine/Source/Developer/ScreenShotComparison/Private/Widgets/SScreenComparisonRow.cpp
@@ -343,7 +343,7 @@ bool SScreenComparisonRow::IsComparingAgainstPlatformFallback() const
 TSharedRef<SWidget> SScreenComparisonRow::BuildAddedView()
 {
 	const FImageComparisonResult& ComparisonResult = Model->Report.GetComparisonResult();
-	FString IncomingFile = FPaths::Combine(Model->Report.GetReportPath(), ComparisonResult.ReportIncomingFilePath);
+	FString IncomingFile = FPaths::Combine(Model->Report.GetReportRootDirectory(), ComparisonResult.ReportIncomingFilePath);
 
 	return
 		SNew(SVerticalBox)
@@ -410,7 +410,7 @@ TSharedRef<SWidget> SScreenComparisonRow::BuildComparisonPreview()
 {
 	const FImageComparisonResult& ComparisonResult = Model->Report.GetComparisonResult();
 
-	FString ApprovedFile = FPaths::Combine(Model->Report.GetReportPath(), ComparisonResult.ReportApprovedFilePath);
+	FString ApprovedFile = FPaths::Combine(Model->Report.GetReportRootDirectory(), ComparisonResult.ReportApprovedFilePath);
 
 	// If the actual approved file is on disk then use that so the tool-tip is more useful
 	if (IFileManager::Get().FileExists(*ComparisonResult.ApprovedFilePath))
@@ -418,8 +418,8 @@ TSharedRef<SWidget> SScreenComparisonRow::BuildComparisonPreview()
 		ApprovedFile = ComparisonResult.ApprovedFilePath;
 	}
 
-	FString IncomingFile = FPaths::Combine(Model->Report.GetReportPath(), ComparisonResult.ReportIncomingFilePath);
-	FString DeltaFile = FPaths::Combine(Model->Report.GetReportPath(), ComparisonResult.ReportComparisonFilePath);
+	FString IncomingFile = FPaths::Combine(Model->Report.GetReportRootDirectory(), ComparisonResult.ReportIncomingFilePath);
+	FString DeltaFile = FPaths::Combine(Model->Report.GetReportRootDirectory(), ComparisonResult.ReportComparisonFilePath);
 
 	// Create the screen shot data widget.
 	return 
@@ -596,7 +596,7 @@ FReply SScreenComparisonRow::OnCompareNewImage(const FGeometry& InGeometry, cons
 {
 
 	const FImageComparisonResult& ComparisonResult = Model->Report.GetComparisonResult();
-	FString IncomingFilePath = FPaths::Combine(Model->Report.GetReportPath(), ComparisonResult.ReportIncomingFilePath);
+	FString IncomingFilePath = FPaths::Combine(Model->Report.GetReportRootDirectory(), ComparisonResult.ReportIncomingFilePath);
 
 	TSharedPtr<FSlateDynamicImageBrush> UnapprovedImage = UnapprovedImageWidget->GetDynamicBrush();
 
diff --git a/Engine/Source/Developer/ScreenShotComparisonTools/Private/ScreenShotManager.cpp b/Engine/Source/Developer/ScreenShotComparisonTools/Private/ScreenShotManager.cpp
--- a/Engine/Source/Developer/ScreenShotComparisonTools/Private/ScreenShotManager.cpp
+++ b/Engine/Source/Developer/ScreenShotComparisonTools/Private/ScreenShotManager.cpp
@@ -464,9 +464,9 @@ FImageComparisonResult FScreenShotManager::CompareScreenshot(const FString& InUn
 		
 	// make all these paths relative to the report path so the report can be relocated outside of the project directory
 	// and the files still found.
-	FPaths::MakePathRelativeTo(ComparisonResult.ReportApprovedFilePath, *ReportPathOnDisk);
-	FPaths::MakePathRelativeTo(ComparisonResult.ReportIncomingFilePath, *ReportPathOnDisk);
-	FPaths::MakePathRelativeTo(ComparisonResult.ReportComparisonFilePath, *ReportPathOnDisk);
+	FPaths::MakePathRelativeTo(ComparisonResult.ReportApprovedFilePath, *ScreenshotResultsFolder);
+	FPaths::MakePathRelativeTo(ComparisonResult.ReportIncomingFilePath, *ScreenshotResultsFolder);
+	FPaths::MakePathRelativeTo(ComparisonResult.ReportComparisonFilePath, *ScreenshotResultsFolder);
 
 	// save the result at to a report
 	FString Json;
@@ -493,9 +493,9 @@ FScreenshotExportResult FScreenShotManager::ExportScreenshotComparisonResult(FSt
 
 	FScreenshotExportResult Results;
 	Results.Success = false;
-	Results.ExportPath = RootExportFolder / FString::FromInt(FEngineVersion::Current().GetChangelist());
+	Results.ExportPath = RootExportFolder;
 
-	FString Destination = Results.ExportPath / ScreenshotName;
+	FString Destination = RootExportFolder / ScreenshotName;
 	if (!IFileManager::Get().MakeDirectory(*Destination, /*Tree =*/true))
 	{
 		return Results;
@@ -553,7 +553,7 @@ bool FScreenShotManager::OpenComparisonReports(FString ImportPath, TArray<FCompa
 
 FString FScreenShotManager::GetDefaultExportDirectory() const
 {
-	return FPaths::Combine(FPaths::ProjectSavedDir(),TEXT("Exported/"));
+	return FPaths::Combine(FPaths::ProjectSavedDir(),TEXT("Exported/imageCompare"));
 }
 
 void FScreenShotManager::CopyDirectory(const FString& DestDir, const FString& SrcDir)
diff --git a/Engine/Source/Programs/AutomationTool/Gauntlet/Framework/Utils/Gauntlet.HordeReport.cs b/Engine/Source/Programs/AutomationTool/Gauntlet/Framework/Utils/Gauntlet.HordeReport.cs
--- a/Engine/Source/Programs/AutomationTool/Gauntlet/Framework/Utils/Gauntlet.HordeReport.cs
+++ b/Engine/Source/Programs/AutomationTool/Gauntlet/Framework/Utils/Gauntlet.HordeReport.cs
@@ -359,9 +359,8 @@ namespace Gauntlet
 			/// Convert UnrealAutomatedTestPassResults to Horde data model
 			/// </summary>
 			/// <param name="InTestPassResults"></param>
-			/// <param name="ReportPath"></param>
 			/// <param name="ReportURL"></param>
-			public static UnrealEngineTestPassResults FromUnrealAutomatedTests(UnrealAutomatedTestPassResults InTestPassResults, string ReportPath, string ReportURL)
+			public static UnrealEngineTestPassResults FromUnrealAutomatedTests(UnrealAutomatedTestPassResults InTestPassResults, string ReportURL)
 			{
 				UnrealEngineTestPassResults OutTestPassResults = new UnrealEngineTestPassResults();
 				if (InTestPassResults.Devices != null)
@@ -412,9 +411,9 @@ namespace Gauntlet
 							NewArtifact.Name = InTestArtifact.Name;
 							NewArtifact.Type = InTestArtifact.Type;
 							ComparisonFiles ArtifactFiles = NewArtifact.Files;
-							ArtifactFiles.Difference = !string.IsNullOrEmpty(InTestArtifact.Files.Difference)?Path.Combine(ReportPath, InTestArtifact.Files.Difference):null;
-							ArtifactFiles.Approved = !string.IsNullOrEmpty(InTestArtifact.Files.Approved)?Path.Combine(ReportPath, InTestArtifact.Files.Approved):null;
-							ArtifactFiles.Unapproved = !string.IsNullOrEmpty(InTestArtifact.Files.Unapproved)?Path.Combine(ReportPath, InTestArtifact.Files.Unapproved):null;
+							ArtifactFiles.Difference = InTestArtifact.Files.Difference;
+							ArtifactFiles.Approved = InTestArtifact.Files.Approved;
+							ArtifactFiles.Unapproved = InTestArtifact.Files.Unapproved;
 						}
 						foreach (UnrealAutomationEntry InTestEntry in InTestResult.Entries)
 						{
@@ -436,8 +435,9 @@ namespace Gauntlet
 			/// <summary>
 			/// Copy Test Results Artifacts
 			/// </summary>
+			/// <param name="ReportPath"></param>
 			/// <param name="OutputArtifactPath"></param>
-			public void CopyTestResultsArtifacts(string InOutputArtifactPath)
+			public void CopyTestResultsArtifacts(string ReportPath, string InOutputArtifactPath)
 			{
 				SetOutputArtifactPath(InOutputArtifactPath);
 				int ArtifactsCount = 0;
@@ -447,14 +447,12 @@ namespace Gauntlet
 					// copy artifacts
 					foreach (Artifact TestArtifact in OutputTestResultDetailed.Artifacts)
 					{
+						
 						string[] ArtifactPaths= { TestArtifact.Files.Difference, TestArtifact.Files.Approved, TestArtifact.Files.Unapproved };
 						foreach (string ArtifactPath in ArtifactPaths)
 						{
-							if (AttachArtifact(ArtifactPath)) { ArtifactsCount++; }
+							if (!string.IsNullOrEmpty(ArtifactPath) && AttachArtifact(Path.Combine(ReportPath, ArtifactPath), ArtifactPath)) { ArtifactsCount++; }
 						}
-						TestArtifact.Files.Difference = Path.GetFileName(TestArtifact.Files.Difference);
-						TestArtifact.Files.Approved = Path.GetFileName(TestArtifact.Files.Approved);
-						TestArtifact.Files.Unapproved = Path.GetFileName(TestArtifact.Files.Unapproved);
 					}
 					// write test json blob
 					string TestResultFilePath = Path.Combine(OutputArtifactPath, OutputTestResult.ArtifactName);
diff --git a/Engine/Source/Programs/AutomationTool/Gauntlet/Unreal/Base/Gauntlet.UnrealTestNode.cs b/Engine/Source/Programs/AutomationTool/Gauntlet/Unreal/Base/Gauntlet.UnrealTestNode.cs
--- a/Engine/Source/Programs/AutomationTool/Gauntlet/Unreal/Base/Gauntlet.UnrealTestNode.cs
+++ b/Engine/Source/Programs/AutomationTool/Gauntlet/Unreal/Base/Gauntlet.UnrealTestNode.cs
@@ -1251,8 +1251,8 @@ namespace Gauntlet
 					}
 				}
 				// Convert test results for Horde
-				HordeReport.UnrealEngineTestPassResults HordeTestPassResults = HordeReport.UnrealEngineTestPassResults.FromUnrealAutomatedTests(JsonTestPassResults, UnrealAutomatedTestReportPath, ReportURL);
-				HordeTestPassResults.CopyTestResultsArtifacts(HordeArtifactPath);
+				HordeReport.UnrealEngineTestPassResults HordeTestPassResults = HordeReport.UnrealEngineTestPassResults.FromUnrealAutomatedTests(JsonTestPassResults, ReportURL);
+				HordeTestPassResults.CopyTestResultsArtifacts(UnrealAutomatedTestReportPath, HordeArtifactPath);
 				// Metadata
 				// With UE Test Automation, we care only for one role.
 				var MainRole = new List<UnrealTestRole>() { GetConfiguration().GetMainRequiredRole() };
