Fix Approved and Unapproved screenshots being mixed up in the json reports

diff --git a/Engine/Source/Developer/AutomationController/Private/AutomationControllerManger.cpp b/Engine/Source/Developer/AutomationController/Private/AutomationControllerManger.cpp
--- a/Engine/Source/Developer/AutomationController/Private/AutomationControllerManger.cpp
+++ b/Engine/Source/Developer/AutomationController/Private/AutomationControllerManger.cpp
@@ -445,12 +445,12 @@ void FAutomationControllerManager::ProcessComparisonQueue()
 				FString ScreenshotResultsFolder = FPaths::AutomationReportsDir();
 	
 				// Paths in the result are relative to the automation report directory.	
-				LocalFiles.Add(TEXT("unapproved"), FPaths::Combine(ScreenshotResultsFolder, Result.ReportApprovedFilePath));
+				LocalFiles.Add(TEXT("unapproved"), FPaths::Combine(ScreenshotResultsFolder, Result.ReportIncomingFilePath));
 
 				// unapproved should always be valid. but approved/difference may be empty if this is a new screenshot
 				if (Result.ReportIncomingFilePath.Len())
 				{
-					LocalFiles.Add(TEXT("approved"), FPaths::Combine(ScreenshotResultsFolder, Result.ReportIncomingFilePath));
+					LocalFiles.Add(TEXT("approved"), FPaths::Combine(ScreenshotResultsFolder, Result.ReportApprovedFilePath));
 				}
 
 				if (Result.ReportComparisonFilePath.Len())
