GLog->Flush to ensure the UE_LOGs end up in the log, as we require the output of the above logs for a test to pass

diff --git a/Engine/Source/Developer/AutomationController/Private/AutomationCommandline.cpp b/Engine/Source/Developer/AutomationController/Private/AutomationCommandline.cpp
--- a/Engine/Source/Developer/AutomationController/Private/AutomationCommandline.cpp
+++ b/Engine/Source/Developer/AutomationController/Private/AutomationCommandline.cpp
@@ -478,6 +478,10 @@ public:
 					}
 					UE_LOG(LogAutomationCommandLine, Log, TEXT("Shutting down. GIsCriticalError=%d"), GIsCriticalError);
 					UE_LOG(LogAutomationCommandLine, Display, TEXT("**** TEST COMPLETE. EXIT CODE: %d ****"), GIsCriticalError ? -1 : 0);
+
+					// Lets flush here as we are about to request exit force=true, which on some platforms may not have flushed the UE_LOG buffer
+					GLog->Flush();
+
 					FPlatformMisc::RequestExitWithStatus(true, GIsCriticalError ? -1 : 0);
 					AutomationTestState = EAutomationTestState::Complete;
 				}
