Move logic into platform request exit area to flush logs if we can

diff --git a/Engine/Source/Developer/AutomationController/Private/AutomationCommandline.cpp b/Engine/Source/Developer/AutomationController/Private/AutomationCommandline.cpp
--- a/Engine/Source/Developer/AutomationController/Private/AutomationCommandline.cpp
+++ b/Engine/Source/Developer/AutomationController/Private/AutomationCommandline.cpp
@@ -478,10 +478,6 @@ public:
 					}
 					UE_LOG(LogAutomationCommandLine, Log, TEXT("Shutting down. GIsCriticalError=%d"), GIsCriticalError);
 					UE_LOG(LogAutomationCommandLine, Display, TEXT("**** TEST COMPLETE. EXIT CODE: %d ****"), GIsCriticalError ? -1 : 0);
-
-					// Lets flush here as we are about to request exit force=true, which on some platforms may not have flushed the UE_LOG buffer
-					GLog->Flush();
-
 					FPlatformMisc::RequestExitWithStatus(true, GIsCriticalError ? -1 : 0);
 					AutomationTestState = EAutomationTestState::Complete;
 				}
diff --git a/Engine/Source/Runtime/Core/Private/Unix/UnixPlatformMisc.cpp b/Engine/Source/Runtime/Core/Private/Unix/UnixPlatformMisc.cpp
--- a/Engine/Source/Runtime/Core/Private/Unix/UnixPlatformMisc.cpp
+++ b/Engine/Source/Runtime/Core/Private/Unix/UnixPlatformMisc.cpp
@@ -376,6 +376,14 @@ void FUnixPlatformMisc::RequestExit(bool Force)
 
 	if(Force)
 	{
+		// Make sure the log is flushed.
+		if (!GEnteredSignalHandler && GLog)
+		{
+			// This may be called from other thread, so set this thread as the master.
+			GLog->SetCurrentThreadAsMasterThread();
+			GLog->TearDown();
+		}
+
 		// Force immediate exit. Cannot call abort() here, because abort() raises SIGABRT which we treat as crash
 		// (to prevent other, particularly third party libs, from quitting without us noticing)
 		// Propagate override return code, but normally don't exit with 0, so the parent knows it wasn't a normal exit.
