Test Groups

diff --git a/Engine/Source/Developer/AutomationController/Private/AutomationCommandline.cpp b/Engine/Source/Developer/AutomationController/Private/AutomationCommandline.cpp
--- a/Engine/Source/Developer/AutomationController/Private/AutomationCommandline.cpp
+++ b/Engine/Source/Developer/AutomationController/Private/AutomationCommandline.cpp
@@ -15,6 +15,7 @@
 #include "Misc/FileHelper.h"
 #include "AssetRegistry/AssetRegistryModule.h"
 #include "AutomationControllerSettings.h"
+#include "AutomationGroupFilter.h"
 #include "Containers/Ticker.h"
 
 DEFINE_LOG_CATEGORY_STATIC(LogAutomationCommandLine, Log, All);
@@ -124,18 +125,14 @@ public:
 		return false;
 	}
 
-
-	
-	void GenerateTestNamesFromCommandLine(const TArray<FString>& AllTestNames, TArray<FString>& OutTestNames)
+	void GenerateTestNamesFromCommandLine(TSharedPtr <AutomationFilterCollection> InFilters, TArray<FString>& OutFilteredTestNames)
 	{
-		OutTestNames.Empty();
+		OutFilteredTestNames.Empty();
 		
 		//Split the argument names up on +
 		TArray<FString> ArgumentNames;
 		StringCommand.ParseIntoArray(ArgumentNames, TEXT("+"), true);
 
-		TArray<FAutomatedTestFilter> Filters;
-
 		// get our settings CDO where things are stored
 		UAutomationControllerSettings* Settings = UAutomationControllerSettings::StaticClass()->GetDefaultObject<UAutomationControllerSettings>();
 
@@ -160,7 +157,8 @@ public:
 					FilterName += TEXT(".");
 				}
 
-				Filters.Add(FAutomatedTestFilter(FilterName, true));
+				FAutomationGroupFilter* Filter = new FAutomationGroupFilter(FAutomatedTestFilter(FilterName, true, false));
+				InFilters->Add(MakeShareable(Filter));
 			}
 			else if (ArgumentName.StartsWith(GroupPrefix))
 			{
@@ -178,10 +176,8 @@ public:
 						// if found add all this groups filters to our current list
 						if (GroupEntry->Filters.Num() > 0)
 						{
-							for (const FAutomatedTestFilter& GroupFilter : GroupEntry->Filters)
-							{
-								Filters.Add(GroupFilter);
-							}
+							FAutomationGroupFilter* Filter = new FAutomationGroupFilter(GroupEntry->Filters);
+							InFilters->Add(MakeShareable(Filter));
 						}
 						else
 						{
@@ -211,50 +207,15 @@ public:
 					ArgumentName.LeftChopInline(1);
 				}
 
-				Filters.Add(FAutomatedTestFilter(ArgumentName, bMatchFromStart, bMatchFromEnd));
-			}
-		}
-		
-		for (int32 TestIndex = 0; TestIndex < AllTestNames.Num(); ++TestIndex)
-		{
-			FString TestName = AllTestNames[TestIndex];
-
-			for (const FAutomatedTestFilter& Filter : Filters)
-			{
-				FString FilterString = Filter.Contains;
-
-				bool bNeedStartMatch = Filter.MatchFromStart;
-				bool bNeedEndMatch = Filter.MatchFromEnd;
-				bool bMeetsMatch = true;	// assume true
-
-				// If we need to match at the start or end, 
-				if (bNeedStartMatch || bNeedEndMatch)
-				{
-					if (bNeedStartMatch)
-					{
-						bMeetsMatch = TestName.StartsWith(FilterString);
-					}
-
-					if (bNeedEndMatch && bMeetsMatch)
-					{
-						bMeetsMatch = TestName.EndsWith(FilterString);
-					}
-				}
-				else
-				{
-					// match anywhere
-					bMeetsMatch = TestName.Contains(FilterString);
-				}
-
-				if (bMeetsMatch)
-				{
-					OutTestNames.Add(TestName);
-					TestCount++;
-					break;
-				}
+				FAutomationGroupFilter* Filter = new FAutomationGroupFilter(FAutomatedTestFilter(ArgumentName, bMatchFromStart, bMatchFromEnd));
+				InFilters->Add(MakeShareable(Filter));
 			}
 		}
 
+		// SetFilter applies all filters from the AutomationFilters array
+		AutomationController->SetFilter(InFilters);
+		// Fill OutFilteredTestNames array with filtered test names
+		AutomationController->GetFilteredTestNames(OutFilteredTestNames);
 	}
 
 	void FindWorkers(float DeltaTime)
@@ -312,7 +273,8 @@ public:
 
 		// We have found some workers
 		// Create a filter to add to the automation controller, otherwise we don't get any reports
-		AutomationController->SetFilter(MakeShareable(new AutomationFilterCollection()));
+		TSharedPtr <AutomationFilterCollection> AutomationFilters = MakeShareable(new AutomationFilterCollection());
+		AutomationController->SetFilter(AutomationFilters);
 		AutomationController->SetVisibleTestsEnabled(true);
 		AutomationController->GetEnabledTestNames(AllTestNames);
 
@@ -333,7 +295,7 @@ public:
 		else if (AutomationCommand == EAutomationCommand::RunCommandLineTests)
 		{
 			TArray<FString> FilteredTestNames;
-			GenerateTestNamesFromCommandLine(AllTestNames, FilteredTestNames);
+			GenerateTestNamesFromCommandLine(AutomationFilters, FilteredTestNames);
 			
 			if (FilteredTestNames.Num() == 0)
 			{
diff --git a/Engine/Source/Developer/AutomationController/Private/AutomationControllerManager.cpp b/Engine/Source/Developer/AutomationController/Private/AutomationControllerManager.cpp
new file mode 100644
--- /dev/null
+++ b/Engine/Source/Developer/AutomationController/Private/AutomationControllerManager.cpp
@@ -0,0 +1,1540 @@
+// Copyright Epic Games, Inc. All Rights Reserved.
+
+#include "AutomationControllerManager.h"
+#include "CoreMinimal.h"
+#include "HAL/FileManager.h"
+#include "Misc/CommandLine.h"
+#include "Misc/Paths.h"
+#include "Misc/ConfigCacheIni.h"
+#include "Misc/AutomationTest.h"
+#include "Misc/App.h"
+#include "IAutomationReport.h"
+#include "AutomationWorkerMessages.h"
+#include "IMessageContext.h"
+#include "MessageEndpoint.h"
+#include "MessageEndpointBuilder.h"
+#include "Modules/ModuleManager.h"
+#include "AssetEditorMessages.h"
+#include "ImageComparer.h"
+#include "AutomationControllerSettings.h"
+#include "Interfaces/IScreenShotToolsModule.h"
+#include "Serialization/JsonSerializer.h"
+#include "JsonObjectConverter.h"
+#include "Misc/EngineVersion.h"
+#include "Misc/FileHelper.h"
+#include "PlatformHttp.h"
+
+#include "AutomationTelemetry.h"
+
+#if WITH_EDITOR
+#include "UnrealEdGlobals.h"
+#include "Editor/UnrealEdEngine.h"
+#include "Logging/MessageLog.h"
+#endif
+#include "Async/ParallelFor.h"
+
+// these strings are parsed by Gauntlet (AutomationLogParser) so make sure changes are replicated there!
+#define AutomationTestStarting		TEXT("Test Started. Name={%s} Path={%s}")
+#define AutomationStateFormat		TEXT("Test Completed. Result={%s} Name={%s} Path={%s}")
+
+#define BeginEventsFormat			TEXT("BeginEvents: %s")
+#define EndEventsFormat				TEXT("EndEvents: %s")
+
+DEFINE_LOG_CATEGORY_STATIC(LogAutomationController, Log, All)
+
+#define LOCTEXT_NAMESPACE "AutomationTesting"
+
+void FAutomatedTestPassResults::AddTestResult(const IAutomationReportPtr& TestReport)
+{
+	const FString& FullTestPath = TestReport->GetFullTestPath();
+	if (!TestsMapIndex.Contains(FullTestPath))
+	{
+		FAutomatedTestResult TestResult;
+		TestResult.Test = TestReport;
+		TestResult.TestDisplayName = TestReport->GetDisplayName();
+		TestResult.FullTestPath = FullTestPath;
+
+		TestsMapIndex.Add(FullTestPath, Tests.Num());
+		Tests.Add(TestResult);
+		NotRun++;
+	}
+}
+
+FAutomatedTestResult& FAutomatedTestPassResults::GetTestResult(const IAutomationReportPtr& TestReport)
+{
+	const FString& FullTestPath = TestReport->GetFullTestPath();
+	check(TestsMapIndex.Contains(FullTestPath));
+	return Tests[TestsMapIndex[FullTestPath]];
+}
+
+void FAutomatedTestPassResults::ReBuildTestsMapIndex()
+{
+	TestsMapIndex.Empty();
+	uint32 Index = 0;
+	for (const FAutomatedTestResult& TestResult : Tests)
+	{
+		TestsMapIndex.Add(TestResult.FullTestPath, Index++);
+	}
+}
+
+bool FAutomatedTestPassResults::ReflectResultStateToReport(IAutomationReportPtr& TestReport)
+{
+	const FString& FullTestPath = TestReport->GetFullTestPath();
+	if (TestsMapIndex.Contains(FullTestPath))
+	{
+		FAutomatedTestResult& TestResult = Tests[TestsMapIndex[FullTestPath]];
+		if (TestResult.State == EAutomationState::InProcess)
+		{
+			// Marking incomplete test from previous pass as failure.
+			TestResult.AddEvent(EAutomationEventType::Error, TEXT("Test did not run until completion. The test exited prematurely."));
+			TestResult.State = EAutomationState::Fail;
+			InProcess--;
+			Failed++;
+		}
+		TestReport->SetState(TestResult.State);
+		TestResult.Test = TestReport;
+		return true;
+	}
+	return false;
+}
+
+void FAutomatedTestPassResults::UpdateTestResultStatus(const IAutomationReportPtr& TestReport, EAutomationState State, bool bHasWarning)
+{
+	FAutomatedTestResult& TestResult = GetTestResult(TestReport);
+	if (TestResult.State == State)
+		return;
+
+	// Book keeping
+	// from transient states
+	switch (TestResult.State)
+	{
+	case EAutomationState::InProcess:
+		InProcess--;
+		break;
+	case EAutomationState::NotRun:
+		NotRun--;
+		break;
+	}
+	// to new state
+	switch (State)
+	{
+	case EAutomationState::Success:
+		if (bHasWarning)
+		{
+			SucceededWithWarnings++;
+		}
+		else
+		{
+			Succeeded++;
+		}
+		break;
+	case EAutomationState::Fail:
+		Failed++;
+		break;
+	case EAutomationState::InProcess:
+		TestResult.DateTime = FDateTime::Now();
+		InProcess++;
+		break;
+	case EAutomationState::NotRun:
+		NotRun++;
+		break;
+	case EAutomationState::Skipped:
+		// Those are not counted
+		break;
+	}
+	// commit the change
+	TestResult.State = State;
+}
+
+FAutomationControllerManager::FAutomationControllerManager()
+{
+
+	UAutomationControllerSettings* Settings = UAutomationControllerSettings::StaticClass()->GetDefaultObject<UAutomationControllerSettings>();
+
+	if (Settings->CheckTestIntervalSeconds > 0.0f)
+	{
+		CheckTestIntervalSeconds = Settings->CheckTestIntervalSeconds;
+	}
+	
+	if (Settings->GameInstanceLostTimerSeconds > 0.0f)
+	{
+		GameInstanceLostTimerSeconds = Settings->GameInstanceLostTimerSeconds;
+	}
+	
+	FString DeveloperPath;
+	FParse::Value(FCommandLine::Get(), TEXT("ReportOutputPath="), ReportExportPath, false);
+	FParse::Value(FCommandLine::Get(), TEXT("DisplayReportOutputPath="), ReportURLPath, false);
+	FParse::Value(FCommandLine::Get(), TEXT("DeveloperReportOutputPath="), DeveloperPath, false);
+
+	if (ReportExportPath.Len())
+	{
+		UE_LOG(LogAutomationController, Warning, TEXT("Argument -ReportOutputPath= is now -ReportExportPath=. Please update your command line!"));
+	}
+
+	if (ReportURLPath.Len())
+	{
+		UE_LOG(LogAutomationController, Warning, TEXT("Argument -DisplayReportOutputPath= is now -ReportURL=. Please update your command line!"));
+	}
+
+	if (DeveloperPath.Len())
+	{
+		UE_LOG(LogAutomationController, Warning, TEXT("Argument -DeveloperReportOutputPath= is now -DeveloperReport (no value). Please update your command line!"));
+	}
+
+	// read values with new names
+	FParse::Value(FCommandLine::Get(), TEXT("ReportExportPath="), ReportExportPath, false);
+	FParse::Value(FCommandLine::Get(), TEXT("ReportURL="), ReportURLPath, false);
+	bool bUseDeveloperPath = DeveloperPath.Len() > 0 || FParse::Param(FCommandLine::Get(), TEXT("DeveloperReport"));
+
+	if (ReportExportPath.Len() && bUseDeveloperPath)
+	{
+		ReportExportPath = ReportExportPath / TEXT("dev") / FString(FPlatformProcess::UserName()).ToLower();
+	}
+		
+	if (ReportURLPath.Len() && bUseDeveloperPath)
+	{
+		DeveloperReportUrl = ReportURLPath / TEXT("dev") / FString(FPlatformProcess::UserName()).ToLower() / TEXT("index.html");
+	}
+
+	bResumeRunTest = FParse::Param(FCommandLine::Get(), TEXT("ResumeRunTest"));
+}
+
+void FAutomationControllerManager::RequestAvailableWorkers(const FGuid& SessionId)
+{
+	//invalidate previous tests
+	++ExecutionCount;
+	DeviceClusterManager.Reset();
+
+	ControllerResetDelegate.Broadcast();
+
+	// Don't allow reports to be exported
+	bTestResultsAvailable = false;
+
+	//store off active session ID to reject messages that come in from different sessions
+	ActiveSessionId = SessionId;
+
+	//EAutomationTestFlags::FilterMask
+
+	//TODO AUTOMATION - include change list, game, etc, or remove when launcher is integrated
+	int32 ChangelistNumber = 10000;
+	FString ProcessName = TEXT("instance_name");
+
+	MessageEndpoint->Publish(FMessageEndpoint::MakeMessage<FAutomationWorkerFindWorkers>(ChangelistNumber, FApp::GetProjectName(), ProcessName, SessionId), EMessageScope::Network);
+
+	// Reset the check test timers
+	LastTimeUpdateTicked = FPlatformTime::Seconds();
+	CheckTestTimer = 0.f;
+
+	IScreenShotToolsModule& ScreenShotModule = FModuleManager::LoadModuleChecked<IScreenShotToolsModule>("ScreenShotComparisonTools");
+	ScreenshotManager = ScreenShotModule.GetScreenShotManager();
+}
+
+void FAutomationControllerManager::RequestTests()
+{
+	//invalidate incoming results
+	ExecutionCount++;
+	//reset the number of responses we have received
+	RefreshTestResponses = 0;
+
+	ReportManager.Empty();
+
+	for ( int32 ClusterIndex = 0; ClusterIndex < DeviceClusterManager.GetNumClusters(); ++ClusterIndex )
+	{
+		int32 DevicesInCluster = DeviceClusterManager.GetNumDevicesInCluster(ClusterIndex);
+		if ( DevicesInCluster > 0 )
+		{
+			FMessageAddress MessageAddress = DeviceClusterManager.GetDeviceMessageAddress(ClusterIndex, 0);
+
+			UE_LOG(LogAutomationController, Log, TEXT("Requesting test list from %s"), *MessageAddress.ToString());
+
+			//issue tests on appropriate platforms
+			MessageEndpoint->Send(FMessageEndpoint::MakeMessage<FAutomationWorkerRequestTests>(bDeveloperDirectoryIncluded, RequestedTestFlags), MessageAddress);
+		}
+	}
+}
+
+void FAutomationControllerManager::RunTests(const bool bInIsLocalSession)
+{
+	ExecutionCount++;
+	CurrentTestPass = 0;
+	ReportManager.SetCurrentTestPass(CurrentTestPass);
+	ClusterDistributionMask = 0;
+	bTestResultsAvailable = false;
+	TestRunningArray.Empty();
+	bIsLocalSession = bInIsLocalSession;
+
+	// Reset the check test timers
+	LastTimeUpdateTicked = FPlatformTime::Seconds();
+	CheckTestTimer = 0.f;
+
+#if WITH_EDITOR
+	FMessageLog AutomationEditorLog("AutomationTestingLog");
+	FString NewPageName = FString::Printf(TEXT("-----Test Run %d----"), ExecutionCount);
+	FText NewPageNameText = FText::FromString(*NewPageName);
+	AutomationEditorLog.Open();
+	AutomationEditorLog.NewPage(NewPageNameText);
+	AutomationEditorLog.Info(NewPageNameText);
+#endif
+	//reset all tests
+	ReportManager.ResetForExecution(NumTestPasses);
+
+	// Register All tests that we'll need to be exported as json report
+	if (!ReportExportPath.IsEmpty())
+	{
+		JsonTestPassResults.IsRequired = true;
+		TArray<IAutomationReportPtr> TestsToRun = ReportManager.GetEnabledTestReports();
+
+		// Load previous test results
+		if (bResumeRunTest)
+		{
+			FString ReportFileName = FString::Printf(TEXT("%s/index.json"), *ReportExportPath);
+			if (FPaths::FileExists(ReportFileName))
+			{
+				LoadJsonTestPassSummary(ReportFileName, TestsToRun);
+			}
+		}
+
+		for (IAutomationReportPtr& TestReport : TestsToRun)
+		{
+			JsonTestPassResults.AddTestResult(TestReport);
+		}
+		// Get the html index written down as soon as possible so once results are updated, they can be reviewed.
+		GenerateTestPassHtmlIndex();
+	}
+
+	for ( int32 ClusterIndex = 0; ClusterIndex < DeviceClusterManager.GetNumClusters(); ++ClusterIndex )
+	{
+		//enable each device cluster
+		ClusterDistributionMask |= ( 1 << ClusterIndex );
+
+		//for each device in this cluster
+		for ( int32 DeviceIndex = 0; DeviceIndex < DeviceClusterManager.GetNumDevicesInCluster(ClusterIndex); ++DeviceIndex )
+		{
+			//mark the device as idle
+			DeviceClusterManager.SetTest(ClusterIndex, DeviceIndex, NULL);
+
+			// Send command to reset tests (delete local files, etc)
+			FMessageAddress MessageAddress = DeviceClusterManager.GetDeviceMessageAddress(ClusterIndex, DeviceIndex);
+			UE_LOG(LogAutomationController, Log, TEXT("Sending Reset Tests to %s"), *MessageAddress.ToString());
+
+			MessageEndpoint->Send(FMessageEndpoint::MakeMessage<FAutomationWorkerResetTests>(), MessageAddress);
+
+			// Store devices info into the json report.
+			if (JsonTestPassResults.IsRequired)
+			{
+				const FAutomationDeviceInfo& DeviceInfo = DeviceClusterManager.GetDeviceInfo(ClusterIndex, DeviceIndex);
+				JsonTestPassResults.Devices.Add(DeviceInfo);
+			}
+		}
+	}
+	
+	// Clear the logs and transient folders. Reports can be manually cleared by the user in the UI if they so desire
+	UE_LOG(LogAutomationController, Display, TEXT("Clearing %s and %s"), *FPaths::AutomationLogDir(), *FPaths::AutomationTransientDir());
+	IFileManager::Get().DeleteDirectory(*FPaths::AutomationLogDir(), false, true);
+	IFileManager::Get().DeleteDirectory(*FPaths::AutomationTransientDir(), false, true);
+
+
+	// Inform the UI we are running tests
+	if ( ClusterDistributionMask != 0 )
+	{
+		SetControllerStatus(EAutomationControllerModuleState::Running);
+	}
+}
+
+void FAutomationControllerManager::StopTests()
+{
+	bTestResultsAvailable = false;
+	ClusterDistributionMask = 0;
+
+	ReportManager.StopRunningTests();
+
+	// Inform the UI we have stopped running tests
+	if ( DeviceClusterManager.HasActiveDevice() )
+	{
+		for (int32 ClusterIndex = 0; ClusterIndex < DeviceClusterManager.GetNumClusters(); ++ClusterIndex)
+		{
+			//for each device in this cluster
+			for (int32 DeviceIndex = 0; DeviceIndex < DeviceClusterManager.GetNumDevicesInCluster(ClusterIndex); ++DeviceIndex)
+			{
+				//mark the device as idle
+				DeviceClusterManager.SetTest(ClusterIndex, DeviceIndex, NULL);
+
+				// Send command to reset tests (delete local files, etc)
+				FMessageAddress MessageAddress = DeviceClusterManager.GetDeviceMessageAddress(ClusterIndex, DeviceIndex);
+
+				UE_LOG(LogAutomationController, Log, TEXT("Sending StopTests to %s"), *MessageAddress.ToString());
+				MessageEndpoint->Send(FMessageEndpoint::MakeMessage<FAutomationWorkerStopTests>(), MessageAddress);
+			}
+		}
+
+		SetControllerStatus(EAutomationControllerModuleState::Ready);
+	}
+	else
+	{
+		SetControllerStatus(EAutomationControllerModuleState::Disabled);
+	}
+
+	TestRunningArray.Empty();
+
+	// Close play window
+#if WITH_EDITOR
+	if (GUnrealEd)
+	{
+		GUnrealEd->RequestEndPlayMap();
+	}
+#endif
+}
+
+void FAutomationControllerManager::Init()
+{
+	extern void EmptyLinkFunctionForStaticInitializationAutomationExecCmd();
+	EmptyLinkFunctionForStaticInitializationAutomationExecCmd();
+
+	AutomationTestState = EAutomationControllerModuleState::Disabled;
+	bTestResultsAvailable = false;
+	bSendAnalytics = FParse::Param(FCommandLine::Get(), TEXT("SendAutomationAnalytics"));
+}
+
+void FAutomationControllerManager::RequestLoadAsset(const FString& InAssetName)
+{
+	MessageEndpoint->Publish(FMessageEndpoint::MakeMessage<FAssetEditorRequestOpenAsset>(InAssetName), EMessageScope::Process);
+}
+
+void FAutomationControllerManager::Tick()
+{
+	ProcessAvailableTasks();
+	ProcessComparisonQueue();
+}
+
+void FAutomationControllerManager::ProcessComparisonQueue()
+{
+	TSharedPtr<FComparisonEntry> Entry;
+	if ( ComparisonQueue.Peek(Entry) )
+	{
+		if ( Entry->PendingComparison.IsReady() )
+		{
+			const bool Dequeued = ComparisonQueue.Dequeue(Entry);
+			check(Dequeued);
+
+			FImageComparisonResult Result = Entry->PendingComparison.Get();
+
+			const FGuid UniqueId = FGuid::NewGuid();
+
+			// Send the message back to the automation worker letting it know the results of the comparison test.
+			{
+				FAutomationWorkerImageComparisonResults* Message = FMessageEndpoint::MakeMessage<FAutomationWorkerImageComparisonResults>(
+					UniqueId,
+					Result.IsNew(),
+					Result.AreSimilar(),
+					Result.MaxLocalDifference,
+					Result.GlobalDifference,
+					Result.ErrorMessage.ToString()
+				);
+
+				UE_LOG(LogAutomationController, Log, TEXT("Sending ImageComparisonResult to %s (IsNew=%d, AreSimilar=%d)"), 
+					*Entry->Sender.ToString()
+					, Result.IsNew()
+					, Result.AreSimilar()
+					);
+				MessageEndpoint->Send(Message, Entry->Sender);
+			}
+
+			// Find the game session instance info
+			int32 ClusterIndex;
+			int32 DeviceIndex;
+			verify(DeviceClusterManager.FindDevice(Entry->Sender, ClusterIndex, DeviceIndex));
+
+			// Get the current test.
+			TSharedPtr<IAutomationReport> Report = DeviceClusterManager.GetTest(ClusterIndex, DeviceIndex);
+			if (Report.IsValid())
+			{
+				// Record the artifacts for the test.
+				TMap<FString, FString> LocalFiles;
+
+				FString ScreenshotResultsFolder = FPaths::AutomationReportsDir();
+	
+				// Paths in the result are relative to the automation report directory.	
+				LocalFiles.Add(TEXT("unapproved"), FPaths::Combine(ScreenshotResultsFolder, Result.ReportIncomingFilePath));
+
+				// unapproved should always be valid. but approved/difference may be empty if this is a new screenshot
+				if (Result.ReportIncomingFilePath.Len())
+				{
+					LocalFiles.Add(TEXT("approved"), FPaths::Combine(ScreenshotResultsFolder, Result.ReportApprovedFilePath));
+				}
+
+				if (Result.ReportComparisonFilePath.Len())
+				{
+					LocalFiles.Add(TEXT("difference"), FPaths::Combine(ScreenshotResultsFolder, Result.ReportComparisonFilePath));
+				}
+
+				Report->AddArtifact(ClusterIndex, CurrentTestPass, FAutomationArtifact(UniqueId, Entry->ScreenshotPath, EAutomationArtifactType::Comparison, LocalFiles));
+			}
+			else
+			{
+				UE_LOG(LogAutomationController, Error, TEXT("Cannot generate screenshot report for screenshot %s as report is missing"), *Result.IncomingFilePath);
+			}
+		}
+	}
+}
+
+void FAutomationControllerManager::ProcessAvailableTasks()
+{
+	// Distribute tasks
+	if ( ClusterDistributionMask != 0 )
+	{
+		// For each device cluster
+		for ( int32 ClusterIndex = 0; ClusterIndex < DeviceClusterManager.GetNumClusters(); ++ClusterIndex )
+		{
+			bool bAllTestsComplete = true;
+
+			// If any of the devices were valid
+			if ( ( ClusterDistributionMask & ( 1 << ClusterIndex ) ) && DeviceClusterManager.GetNumDevicesInCluster(ClusterIndex) > 0 )
+			{
+				ExecuteNextTask(ClusterIndex, bAllTestsComplete);
+			}
+
+			//if we're all done running our tests
+			if ( bAllTestsComplete )
+			{
+				//we don't need to test this cluster anymore
+				ClusterDistributionMask &= ~( 1 << ClusterIndex );
+
+				if ( ClusterDistributionMask == 0 )
+				{
+					ProcessResults();
+
+					// Close play window
+					#if WITH_EDITOR
+					if (GUnrealEd)
+					{
+						GUnrealEd->RequestEndPlayMap();
+					}
+					#endif
+
+					//Notify the graphical layout we are done processing results.
+					TestsCompleteDelegate.Broadcast();
+				}
+			}
+		}
+	}
+
+	if ( bIsLocalSession == false )
+	{
+		// Update the test status for timeouts if this is not a local session
+		UpdateTests();
+	}
+}
+
+void FAutomationControllerManager::ReportTestResults()
+{
+	UE_LOG(LogAutomationController, Log, TEXT("Test Pass Results:"));
+	for ( int32 i = 0; i < JsonTestPassResults.Tests.Num(); i++ )
+	{
+		UE_LOG(LogAutomationController, Log, TEXT("%s: %s"), *JsonTestPassResults.Tests[i].TestDisplayName, *AutomationStateToString(JsonTestPassResults.Tests[i].State));
+	}
+}
+
+void FAutomationControllerManager::CollectTestResults(TSharedPtr<IAutomationReport> Report, const FAutomationTestResults& Results)
+{
+	if (JsonTestPassResults.IsRequired)
+	{
+		JsonTestPassResults.UpdateTestResultStatus(Report, Results.State, Results.GetWarningTotal() > 0);
+		FAutomatedTestResult& TestResult = JsonTestPassResults.GetTestResult(Report);
+		TestResult.SetEvents(Results.GetEntries(), Results.GetWarningTotal(), Results.GetErrorTotal());
+		TestResult.SetArtifacts(Results.Artifacts);
+		TestResult.Duration = Results.Duration;
+
+		JsonTestPassResults.TotalDuration += Results.Duration;
+
+		FString ImageCompareDirName = TEXT("imageCompare");
+		if (FEngineVersion::Current().GetChangelist() != 0)
+		{
+			ImageCompareDirName = ImageCompareDirName / FString::FromInt(FEngineVersion::Current().GetChangelist());
+		}
+		FString ImageCompareExportPath = ReportExportPath / ImageCompareDirName;
+
+		// Copy new artifacts to Report export path
+		FCriticalSection CS;
+		for (FAutomationArtifact& Artifact : TestResult.GetArtifacts())
+		{
+			TArray<FString> Keys;
+			Artifact.LocalFiles.GetKeys(Keys);
+
+			ParallelFor(Keys.Num(), [&](int32 Index)
+				{
+					const FString& Key = Keys[Index];
+					FString Path = Artifact.LocalFiles[Key];
+					FPaths::MakePathRelativeTo(Path, *FPaths::AutomationReportsDir());
+					Path = ImageCompareDirName / Path;
+					{
+						FScopeLock Lock(&CS);
+						Artifact.Files.Add(Key, MoveTemp(Path));
+					}
+					if (Key == TEXT("unapproved"))
+					{
+						// Copy screenshot report
+						FScreenshotExportResult ExportResult = ScreenshotManager->ExportScreenshotComparisonResult(Artifact.Name, ImageCompareExportPath);
+
+						FScopeLock Lock(&CS);
+						if (!JsonTestPassResults.ComparisonExported && ExportResult.Success)
+						{
+							JsonTestPassResults.ComparisonExported = ExportResult.Success;
+							JsonTestPassResults.ComparisonExportDirectory = ExportResult.ExportPath;
+						}
+					}
+				});
+		}
+	}
+}
+
+bool FAutomationControllerManager::GenerateJsonTestPassSummary(FAutomatedTestPassResults& SerializedPassResults)
+{
+	UE_LOG(LogAutomationController, Display, TEXT("Converting results to json object..."));
+	SerializedPassResults.ReportCreatedOn = FDateTime::Now();
+
+	FString Json;
+	if (FJsonObjectConverter::UStructToJsonObjectString(SerializedPassResults, Json))
+	{
+		FString ReportFileName = FString::Printf(TEXT("%s/index.json"), *ReportExportPath);
+		const int32 WriteAttempts = 3;
+		const float SleepBetweenAttempts = 0.05f;
+
+		for (int32 Attempt = 1; Attempt <= WriteAttempts; ++Attempt)
+		{
+			if (FFileHelper::SaveStringToFile(Json, *ReportFileName, FFileHelper::EEncodingOptions::ForceUTF8))
+			{
+				UE_LOG(LogAutomationController, Display, TEXT("Successfully wrote json results file!"));
+				return true;
+			}
+
+			FPlatformProcess::Sleep(SleepBetweenAttempts);
+		}
+
+		UE_LOG(LogAutomationController, Warning, TEXT("Failed to write test report json to '%s' after 3 attempts - No report will be generated."), *ReportFileName);
+	}
+	else
+	{
+		UE_LOG(LogAutomationController, Error, TEXT("Failed to convert test results to json object - No report will be generated."));
+	}
+		
+	return false;
+}
+
+bool FAutomationControllerManager::GenerateTestPassHtmlIndex()
+{
+	UE_LOG(LogAutomationController, Display, TEXT("Loading results html template..."));
+
+	FString ReportTemplate;
+	if (FFileHelper::LoadFileToString(ReportTemplate, *(FPaths::EngineContentDir() / TEXT("Automation/Report-Template.html"))))
+	{		
+		FString ReportFileName = FString::Printf(TEXT("%s/index.html"), *ReportExportPath);
+		const int32 WriteAttempts = 3;
+		const float SleepBetweenAttempts = 0.05f;
+
+		for (int32 Attempt = 1; Attempt <= WriteAttempts; ++Attempt)
+		{
+			if (FFileHelper::SaveStringToFile(ReportTemplate, *ReportFileName, FFileHelper::EEncodingOptions::ForceUTF8))
+			{
+				UE_LOG(LogAutomationController, Display, TEXT("Successfully wrote html results file!"));
+				return true;
+			}
+
+			FPlatformProcess::Sleep(SleepBetweenAttempts);
+		}
+
+		UE_LOG(LogAutomationController, Warning, TEXT("Failed to write test report html to '%s' after 3 attempts - No report will be generated."), *ReportFileName);
+	}
+	else
+	{
+		UE_LOG(LogAutomationController, Error, TEXT("Failed to load test report html template - No report will be generated."));
+	}
+	
+	return false;
+}
+
+bool FAutomationControllerManager::LoadJsonTestPassSummary(FString& ReportFilePath, TArray<IAutomationReportPtr> TestReports)
+{
+	FString Json;
+	if (!FFileHelper::LoadFileToString(Json, *ReportFilePath))
+	{
+		UE_LOG(LogAutomationController, Error, TEXT("Failed to read json results file '%s'."), *ReportFilePath);
+		return false;
+	}
+	if (!FJsonObjectConverter::JsonObjectStringToUStruct(Json, &JsonTestPassResults))
+	{
+		UE_LOG(LogAutomationController, Error, TEXT("Failed to load json results file '%s'."), *ReportFilePath);
+		return false;
+	}
+	UE_LOG(LogAutomationController, Display, TEXT("Successfully loaded json results file."));
+
+	JsonTestPassResults.ReBuildTestsMapIndex();
+
+	// Reflect Test Result Status on Report Status
+	for (IAutomationReportPtr& TestReport : TestReports)
+	{
+		if (!JsonTestPassResults.ReflectResultStateToReport(TestReport))
+		{
+			UE_LOG(LogAutomationController, Warning, TEXT("No match to test '%s' in json results."), *TestReport->GetFullTestPath());
+		}
+	}
+
+	return true;
+}
+
+FString FAutomationControllerManager::SlugString(const FString& DisplayString) const
+{
+	FString GeneratedName = DisplayString;
+
+	// Convert the display label, which may consist of just about any possible character, into a
+	// suitable name for a UObject (remove whitespace, certain symbols, etc.)
+	{
+		for ( int32 BadCharacterIndex = 0; BadCharacterIndex < UE_ARRAY_COUNT(INVALID_OBJECTNAME_CHARACTERS) - 1; ++BadCharacterIndex )
+		{
+			const TCHAR TestChar[2] = { INVALID_OBJECTNAME_CHARACTERS[BadCharacterIndex], 0 };
+			const int32 NumReplacedChars = GeneratedName.ReplaceInline(TestChar, TEXT(""));
+		}
+	}
+
+	return GeneratedName;
+}
+
+FString FAutomationControllerManager::GetReportOutputPath() const
+{
+	return ReportExportPath;
+}
+
+void FAutomationControllerManager::ExecuteNextTask( int32 ClusterIndex, OUT bool& bAllTestsCompleted )
+{
+	bool bTestThatRequiresMultiplePraticipantsHadEnoughParticipants = false;
+	TArray< IAutomationReportPtr > TestsRunThisPass;
+
+	// For each device in this cluster
+	int32 NumDevicesInCluster = DeviceClusterManager.GetNumDevicesInCluster( ClusterIndex );
+	for ( int32 DeviceIndex = 0; DeviceIndex < NumDevicesInCluster; ++DeviceIndex )
+	{
+		// If this device is idle
+		if ( !DeviceClusterManager.GetTest(ClusterIndex, DeviceIndex).IsValid() && DeviceClusterManager.DeviceEnabled(ClusterIndex, DeviceIndex) )
+		{
+			// Get the next test that should be worked on
+			TSharedPtr< IAutomationReport > NextTest = ReportManager.GetNextReportToExecute(bAllTestsCompleted, ClusterIndex, CurrentTestPass, NumDevicesInCluster);
+			if ( NextTest.IsValid() )
+			{
+				// Get the status of the test
+				EAutomationState TestState = NextTest->GetState(ClusterIndex, CurrentTestPass);
+				if (TestState == EAutomationState::NotRun)
+				{
+					// Reserve this device for the test
+					DeviceClusterManager.SetTest(ClusterIndex, DeviceIndex, NextTest);
+					TestsRunThisPass.Add(NextTest);
+
+					// If we now have enough devices reserved for the test, run it!
+					TArray<FMessageAddress> DeviceAddresses = DeviceClusterManager.GetDevicesReservedForTest(ClusterIndex, NextTest);
+					if (DeviceAddresses.Num() == NextTest->GetNumParticipantsRequired())
+					{
+						TArray<FString> GameInstances;
+
+						// Send it to each device
+						for (int32 AddressIndex = 0; AddressIndex < DeviceAddresses.Num(); ++AddressIndex)
+						{
+							FAutomationTestResults TestResults;
+							TestResults.GameInstance = DeviceClusterManager.GetClusterDeviceName(ClusterIndex, DeviceIndex);
+							TestResults.State = EAutomationState::InProcess;
+							GameInstances.Add(TestResults.GameInstance);
+							NextTest->SetResults(ClusterIndex, CurrentTestPass, TestResults);
+
+							// Mark the device as busy
+							FMessageAddress DeviceAddress = DeviceAddresses[AddressIndex];
+
+							// Send the test to the device for execution!
+							UE_LOG(LogAutomationController, Log, TEXT("Sending RunTest %s to %s"), *NextTest->GetDisplayName(), *DeviceAddress.ToString());
+
+							MessageEndpoint->Send(FMessageEndpoint::MakeMessage<FAutomationWorkerRunTests>(ExecutionCount, AddressIndex, NextTest->GetCommand(), NextTest->GetDisplayName(), NextTest->GetFullTestPath(), bSendAnalytics), DeviceAddress);
+
+							// Add a test so we can check later if the device is still active
+							TestRunningArray.Add(FTestRunningInfo(DeviceAddress));
+						}
+
+						UE_LOG(LogAutomationController, Display, AutomationTestStarting, *NextTest->GetDisplayName(), *NextTest->GetFullTestPath());
+
+						if (JsonTestPassResults.IsRequired && bResumeRunTest)
+						{
+							JsonTestPassResults.UpdateTestResultStatus(NextTest, EAutomationState::InProcess);
+							FAutomatedTestResult& TestResult = JsonTestPassResults.GetTestResult(NextTest);
+							TestResult.DeviceInstance = GameInstances;
+							// Save the whole pass report so that if the next test triggers a critical failure we are not left with no pass report and we can resume.
+							GenerateJsonTestPassSummary(JsonTestPassResults);
+						}
+
+						NextTest->ResetNetworkCommandResponses();
+					}
+				}
+			}
+		}
+		else
+		{
+			// At least one device is still working
+			bAllTestsCompleted = false;
+		}
+	}
+
+	// Ensure any tests we have attempted to run on this pass had enough participants to successfully run.
+	for ( int32 TestIndex = 0; TestIndex < TestsRunThisPass.Num(); TestIndex++ )
+	{
+		IAutomationReportPtr CurrentTest = TestsRunThisPass[TestIndex];
+
+		if ( CurrentTest->GetNumDevicesRunningTest() != CurrentTest->GetNumParticipantsRequired() )
+		{
+			if ( GetNumDevicesInCluster(ClusterIndex) < CurrentTest->GetNumParticipantsRequired() )
+			{
+				FAutomationTestResults TestResults;
+				TestResults.State = EAutomationState::Skipped;
+				TestResults.AddEvent(FAutomationEvent(EAutomationEventType::Warning, FString::Printf(TEXT("Skipped because the test needed %d devices to participate, Only had %d available."), CurrentTest->GetNumParticipantsRequired(), DeviceClusterManager.GetNumDevicesInCluster(ClusterIndex))));
+
+				CurrentTest->SetResults(ClusterIndex, CurrentTestPass, TestResults);
+				CollectTestResults(CurrentTest, TestResults);
+				DeviceClusterManager.ResetAllDevicesRunningTest(ClusterIndex, CurrentTest);
+			}
+		}
+	}
+
+	//Check to see if we finished a pass
+	if ( bAllTestsCompleted && CurrentTestPass < NumTestPasses - 1 )
+	{
+		CurrentTestPass++;
+		ReportManager.SetCurrentTestPass(CurrentTestPass);
+		bAllTestsCompleted = false;
+	}
+}
+
+
+void FAutomationControllerManager::Startup()
+{
+	MessageEndpoint = FMessageEndpoint::Builder("FAutomationControllerModule")
+		.Handling<FAutomationWorkerFindWorkersResponse>(this, &FAutomationControllerManager::HandleFindWorkersResponseMessage)
+		.Handling<FAutomationWorkerPong>(this, &FAutomationControllerManager::HandlePongMessage)
+		.Handling<FAutomationWorkerRequestNextNetworkCommand>(this, &FAutomationControllerManager::HandleRequestNextNetworkCommandMessage)
+		.Handling<FAutomationWorkerRequestTestsReplyComplete>(this, &FAutomationControllerManager::HandleRequestTestsReplyCompleteMessage)
+		.Handling<FAutomationWorkerRunTestsReply>(this, &FAutomationControllerManager::HandleRunTestsReplyMessage)
+		.Handling<FAutomationWorkerScreenImage>(this, &FAutomationControllerManager::HandleReceivedScreenShot)
+		.Handling<FAutomationWorkerTestDataRequest>(this, &FAutomationControllerManager::HandleTestDataRequest)
+		.Handling<FAutomationWorkerWorkerOffline>(this, &FAutomationControllerManager::HandleWorkerOfflineMessage)
+		.Handling<FAutomationWorkerTelemetryData>(this, &FAutomationControllerManager::HandleReceivedTelemetryData);
+
+	if ( MessageEndpoint.IsValid() )
+	{
+		MessageEndpoint->Subscribe<FAutomationWorkerWorkerOffline>();
+	}
+
+	ClusterDistributionMask = 0;
+	ExecutionCount = 0;
+	bDeveloperDirectoryIncluded = false;
+	RequestedTestFlags = EAutomationTestFlags::SmokeFilter | EAutomationTestFlags::EngineFilter | EAutomationTestFlags::ProductFilter | EAutomationTestFlags::PerfFilter;
+
+	NumTestPasses = 1;
+
+	//Default to machine name
+	DeviceGroupFlags = 0;
+	ToggleDeviceGroupFlag(EAutomationDeviceGroupTypes::MachineName);
+}
+
+void FAutomationControllerManager::Shutdown()
+{
+	MessageEndpoint.Reset();
+	ShutdownDelegate.Broadcast();
+	RemoveCallbacks();
+}
+
+void FAutomationControllerManager::RemoveCallbacks()
+{
+	ShutdownDelegate.Clear();
+	TestsAvailableDelegate.Clear();
+	TestsRefreshedDelegate.Clear();
+	TestsCompleteDelegate.Clear();
+}
+
+void FAutomationControllerManager::SetTestNames(const FMessageAddress& AutomationWorkerAddress, TArray<FAutomationTestInfo>& TestInfo)
+{
+	int32 DeviceClusterIndex = INDEX_NONE;
+	int32 DeviceIndex = INDEX_NONE;
+
+	// Find the device that requested these tests
+	if ( DeviceClusterManager.FindDevice(AutomationWorkerAddress, DeviceClusterIndex, DeviceIndex) )
+	{
+		// Sort tests by display name
+		struct FCompareAutomationTestInfo
+		{
+			FORCEINLINE bool operator()(const FAutomationTestInfo& A, const FAutomationTestInfo& B) const
+			{
+				return A.GetDisplayName() < B.GetDisplayName();
+			}
+		};
+
+		TestInfo.Sort(FCompareAutomationTestInfo());
+
+		// Add each test to the collection
+		for ( int32 TestIndex = 0; TestIndex < TestInfo.Num(); ++TestIndex )
+		{
+			// Ensure our test exists. If not, add it
+			ReportManager.EnsureReportExists(TestInfo[TestIndex], DeviceClusterIndex, NumTestPasses);
+		}
+	}
+	else
+	{
+		//todo automation - make sure to report error if the device was not discovered correctly
+	}
+
+	// Note the response
+	RefreshTestResponses++;
+
+	// If we have received all the responses we expect to
+	if ( RefreshTestResponses == DeviceClusterManager.GetNumClusters() )
+	{
+		TestsRefreshedDelegate.Broadcast();
+	}
+}
+
+void FAutomationControllerManager::ProcessResults()
+{
+	bHasErrors = false;
+	bHasWarning = false;
+	bHasLogs = false;
+
+	TArray< TSharedPtr< IAutomationReport > >& TestReports = GetReports();
+
+	if ( TestReports.Num() )
+	{
+		bTestResultsAvailable = true;
+
+		for ( int32 Index = 0; Index < TestReports.Num(); Index++ )
+		{
+			CheckChildResult(TestReports[Index]);
+		}
+	}
+
+	if ( JsonTestPassResults.IsRequired )
+	{
+		UE_LOG(LogAutomationController, Display, TEXT("Exporting Automation Report to %s."), *ReportExportPath);
+
+		FAutomatedTestPassResults SerializedPassResults = JsonTestPassResults;
+
+		{ 
+			// Sort result by failure to improve readability
+			SerializedPassResults.Tests.StableSort([](const FAutomatedTestResult& A, const FAutomatedTestResult& B) {
+				if (A.GetErrorTotal() > 0)
+				{
+					if (B.GetErrorTotal() > 0)
+						return (A.FullTestPath < B.FullTestPath);
+					else
+						return true;
+				}
+				else if (B.GetErrorTotal() > 0)
+				{
+					return false;
+				}
+
+				if (A.State == EAutomationState::Skipped)
+				{
+					if (B.State == EAutomationState::Skipped)
+						return  (A.FullTestPath < B.FullTestPath);
+					else
+						return true;
+				}
+				else if (B.State == EAutomationState::Skipped)
+				{
+					return false;
+				}
+
+				if (A.GetWarningTotal() > 0)
+				{
+					if (B.GetWarningTotal() > 0)
+						return (A.FullTestPath < B.FullTestPath);
+					else
+						return true;
+				}
+				else if (B.GetWarningTotal() > 0)
+				{
+					return false;
+				}
+
+				return A.FullTestPath < B.FullTestPath;
+			});
+		}
+
+		{
+			FDateTime StepTime = FDateTime::Now();
+			UE_LOG(LogAutomationController, Log, TEXT("Writing reports to %s..."), *ReportExportPath);
+
+			// Generate Json
+			GenerateJsonTestPassSummary(SerializedPassResults);
+
+			UE_LOG(LogAutomationController, Log, TEXT("Exported report to '%s' in %.02f Seconds"), *ReportExportPath, (FDateTime::Now() - StepTime).GetTotalSeconds());
+		}
+
+		if (!ReportURLPath.IsEmpty())
+		{
+			UE_LOG(LogAutomationController, Display, TEXT("Report can be viewed in a browser at '%s'"), DeveloperReportUrl.IsEmpty() ? *ReportURLPath : *DeveloperReportUrl);
+
+			if (!DeveloperReportUrl.IsEmpty())
+			{
+				UE_LOG(LogAutomationController, Display, TEXT("Launching Report URL %s."), *DeveloperReportUrl);
+
+				FPlatformProcess::LaunchURL(*DeveloperReportUrl, nullptr, nullptr);
+			}
+		}
+	}
+
+	UE_LOG(LogAutomationController, Display, TEXT("Report can be opened in the editor at '%s'"), ReportExportPath.IsEmpty() ? *FPaths::ConvertRelativePathToFull(FPaths::AutomationReportsDir()) : *ReportExportPath);
+
+	// Then clean our array for the next pass.
+	JsonTestPassResults.ClearAllEntries();
+
+	SetControllerStatus(EAutomationControllerModuleState::Ready);
+}
+
+void FAutomationControllerManager::CheckChildResult(TSharedPtr<IAutomationReport> InReport)
+{
+	TArray<TSharedPtr<IAutomationReport> >& ChildReports = InReport->GetChildReports();
+
+	if ( ChildReports.Num() > 0 )
+	{
+		for ( int32 Index = 0; Index < ChildReports.Num(); Index++ )
+		{
+			CheckChildResult(ChildReports[Index]);
+		}
+	}
+	else if ( ( bHasErrors && bHasWarning && bHasLogs ) == false && InReport->IsEnabled() )
+	{
+		for ( int32 ClusterIndex = 0; ClusterIndex < GetNumDeviceClusters(); ++ClusterIndex )
+		{
+			FAutomationTestResults TestResults = InReport->GetResults(ClusterIndex, CurrentTestPass);
+
+			if ( TestResults.GetErrorTotal() > 0 )
+			{
+				bHasErrors = true;
+			}
+			if ( TestResults.GetWarningTotal() )
+			{
+				bHasWarning = true;
+			}
+			if ( TestResults.GetLogTotal() )
+			{
+				bHasLogs = true;
+			}
+		}
+	}
+}
+
+void FAutomationControllerManager::SetControllerStatus(EAutomationControllerModuleState::Type InAutomationTestState)
+{
+	if ( InAutomationTestState != AutomationTestState )
+	{
+		// Inform the UI if the test state has changed
+		AutomationTestState = InAutomationTestState;
+		TestsAvailableDelegate.Broadcast(AutomationTestState);
+	}
+}
+
+void FAutomationControllerManager::RemoveTestRunning(const FMessageAddress& TestAddressToRemove)
+{
+	for ( int32 Index = 0; Index < TestRunningArray.Num(); Index++ )
+	{
+		if ( TestRunningArray[Index].OwnerMessageAddress == TestAddressToRemove )
+		{
+			TestRunningArray.RemoveAt(Index);
+			break;
+		}
+	}
+}
+
+void FAutomationControllerManager::AddPingResult(const FMessageAddress& ResponderAddress)
+{
+	for ( int32 Index = 0; Index < TestRunningArray.Num(); Index++ )
+	{
+		if ( TestRunningArray[Index].OwnerMessageAddress == ResponderAddress )
+		{
+			TestRunningArray[Index].LastPingTime = 0;
+			break;
+		}
+	}
+}
+
+void FAutomationControllerManager::UpdateTests()
+{
+	const double kIgnoreAsFrameHitchDelta = 2.0f;
+	const double TickDelta = FPlatformTime::Seconds() - LastTimeUpdateTicked;
+
+	LastTimeUpdateTicked = FPlatformTime::Seconds();
+
+	// If this tick was very long then don't check the health of tests. If we've been blocked for a while on a load and ping responses haven't yet been processed
+	// then otherwise might add a delta to the ping time that causes the test to look like it's timed out.
+	if (TickDelta < kIgnoreAsFrameHitchDelta)
+	{
+		CheckTestTimer += TickDelta;
+		if (CheckTestTimer > CheckTestIntervalSeconds)
+		{
+			for ( int32 Index = 0; Index < TestRunningArray.Num(); Index++ )
+			{
+				TestRunningArray[Index].LastPingTime += CheckTestTimer;
+
+				if (TestRunningArray[Index].LastPingTime > GameInstanceLostTimerSeconds)
+				{
+					// Find the game session instance info
+					int32 ClusterIndex;
+					int32 DeviceIndex;
+					verify(DeviceClusterManager.FindDevice(TestRunningArray[Index].OwnerMessageAddress, ClusterIndex, DeviceIndex));
+					//verify this device thought it was busy
+					TSharedPtr <IAutomationReport> Report = DeviceClusterManager.GetTest(ClusterIndex, DeviceIndex);
+					check(Report.IsValid());
+					// A dummy array used to report the result
+
+					TArray<FString> EmptyStringArray;
+					TArray<FString> ErrorStringArray;
+					ErrorStringArray.Add(FString(TEXT("Failed")));
+					bHasErrors = true;
+
+					FAutomationTestResults TestResults;
+					TestResults.State = EAutomationState::Fail;
+					TestResults.GameInstance = DeviceClusterManager.GetClusterDeviceName(ClusterIndex, DeviceIndex);
+					TestResults.AddEvent(FAutomationEvent(EAutomationEventType::Error, FString::Printf(TEXT("Timeout waiting for device %s"), *TestResults.GameInstance)));
+
+					UE_LOG(LogAutomationController, Error, TEXT("Removing test and marking failure after timeout waiting for device %s LastPingTime was greater than %.1f"), *TestResults.GameInstance, GameInstanceLostTimerSeconds);
+
+					// Set the results
+					Report->SetResults(ClusterIndex, CurrentTestPass, TestResults);
+					bTestResultsAvailable = true;
+
+					const FAutomationTestResults& FinalResults = Report->GetResults(ClusterIndex, CurrentTestPass);
+
+					// Gather all of the data relevant to this test for our json reporting.
+					CollectTestResults(Report, FinalResults);
+
+					// Disable the device in the cluster so it is not used again
+					DeviceClusterManager.DisableDevice(ClusterIndex, DeviceIndex);
+
+					// Remove the running test
+					TestRunningArray.RemoveAt(Index--);
+
+					// If there are no more devices, set the module state to disabled
+					if ( DeviceClusterManager.HasActiveDevice() == false )
+					{
+						// Process results first to write out the report
+						ProcessResults();
+
+						UE_LOG(LogAutomationController, Display, TEXT("Module disabled"));
+						SetControllerStatus(EAutomationControllerModuleState::Disabled);
+						ClusterDistributionMask = 0;
+					}
+					else
+					{
+						UE_LOG(LogAutomationController, Display, TEXT("Module not disabled. Keep looking."));
+						// Remove the cluster from the mask if there are no active devices left
+						if ( DeviceClusterManager.GetNumActiveDevicesInCluster(ClusterIndex) == 0 )
+						{
+							ClusterDistributionMask &= ~( 1 << ClusterIndex );
+						}
+						if ( TestRunningArray.Num() == 0 )
+						{
+							SetControllerStatus(EAutomationControllerModuleState::Ready);
+						}
+					}
+				}
+				else
+				{
+					MessageEndpoint->Send(FMessageEndpoint::MakeMessage<FAutomationWorkerPing>(), TestRunningArray[Index].OwnerMessageAddress);
+				}
+			}
+			CheckTestTimer = 0.f;
+		}
+	}
+	else
+	{
+		UE_LOG(LogAutomationController, Log, TEXT("Ignoring very large delta of %.02f seconds in calls to FAutomationControllerManager::Tick() and not penalizing unresponsive tests"), TickDelta);
+	}
+}
+
+const bool FAutomationControllerManager::ExportReport(uint32 FileExportTypeMask)
+{
+	return ReportManager.ExportReport(FileExportTypeMask, GetNumDeviceClusters());
+}
+
+bool FAutomationControllerManager::IsTestRunnable(IAutomationReportPtr InReport) const
+{
+	bool bIsRunnable = false;
+
+	for ( int32 ClusterIndex = 0; ClusterIndex < GetNumDeviceClusters(); ++ClusterIndex )
+	{
+		if ( InReport->IsSupported(ClusterIndex) )
+		{
+			if ( GetNumDevicesInCluster(ClusterIndex) >= InReport->GetNumParticipantsRequired() )
+			{
+				bIsRunnable = true;
+				break;
+			}
+		}
+	}
+
+	return bIsRunnable;
+}
+
+/* FAutomationControllerModule callbacks
+ *****************************************************************************/
+
+void FAutomationControllerManager::HandleFindWorkersResponseMessage(const FAutomationWorkerFindWorkersResponse& Message, const TSharedRef<IMessageContext, ESPMode::ThreadSafe>& Context)
+{
+	UE_LOG(LogAutomationController, Log, TEXT("Received FindWorkersResponseMessage from %s"), *Context->GetSender().ToString());
+
+	if ( Message.SessionId == ActiveSessionId )
+	{
+		DeviceClusterManager.AddDeviceFromMessage(Context->GetSender(), Message, DeviceGroupFlags);
+	}
+
+	RequestTests();
+
+	SetControllerStatus(EAutomationControllerModuleState::Ready);
+}
+
+void FAutomationControllerManager::HandlePongMessage( const FAutomationWorkerPong& Message, const TSharedRef<IMessageContext, ESPMode::ThreadSafe>& Context )
+{
+	AddPingResult(Context->GetSender());
+}
+
+void FAutomationControllerManager::HandleReceivedScreenShot(const FAutomationWorkerScreenImage& Message, const TSharedRef<IMessageContext, ESPMode::ThreadSafe>& Context)
+{
+	UE_LOG(LogAutomationController, Log, TEXT("ReceivedScreenShot for %s (%dx%d) from %s"), *Message.Metadata.TestName, Message.Metadata.Width, Message.Metadata.Height, *Context->GetSender().ToString());
+
+	bool bTree = true;
+
+	FString OutputSubFolder = FPaths::Combine(Message.Metadata.Context, Message.Metadata.ScreenShotName);
+	FString OutputFolder = FPaths::Combine(FPaths::AutomationTransientDir(), FPaths::GetPath(Message.ScreenShotName));
+
+	FString IncomingFileName = FPaths::Combine(OutputFolder,FPaths::GetCleanFilename(Message.ScreenShotName));
+	FString TraceFileName = FPaths::ChangeExtension(IncomingFileName, TEXT(".rdc"));
+
+	FString DirectoryPath = FPaths::GetPath(IncomingFileName);
+
+	if (!IFileManager::Get().MakeDirectory(*DirectoryPath, bTree))
+	{
+		UE_LOG(LogAutomationController, Error, TEXT("Failed to create directory %s for incoming screenshot"), *DirectoryPath);
+		return;
+	}
+	
+	if (!FFileHelper::SaveArrayToFile(Message.ScreenImage, *IncomingFileName))
+	{
+		uint32 WriteErrorCode = FPlatformMisc::GetLastError();
+		TCHAR WriteErrorBuffer[2048];
+		FPlatformMisc::GetSystemErrorMessage(WriteErrorBuffer, 2048, WriteErrorCode);
+		UE_LOG(LogAutomationController, Warning, TEXT("Fail to save screenshot to %s. WriteError: %u (%s)"), *IncomingFileName, WriteErrorCode, WriteErrorBuffer);
+		return;
+	}
+
+	if (Message.FrameTrace.Num() > 0)
+	{
+		if (!FFileHelper::SaveArrayToFile(Message.FrameTrace, *TraceFileName))
+		{
+			uint32 WriteErrorCode = FPlatformMisc::GetLastError();
+			TCHAR WriteErrorBuffer[2048];
+			FPlatformMisc::GetSystemErrorMessage(WriteErrorBuffer, 2048, WriteErrorCode);
+			UE_LOG(LogAutomationController, Warning, TEXT("Faild to save frame trace to %s. WriteError: %u (%s)"), *TraceFileName, WriteErrorCode, WriteErrorBuffer);
+		}
+	}
+
+	// TODO Automation There is identical code in, Engine\Source\Runtime\AutomationWorker\Private\AutomationWorkerModule.cpp,
+	// need to move this code into common area.
+
+	FString Json;
+	if ( FJsonObjectConverter::UStructToJsonObjectString(Message.Metadata, Json) )
+	{
+		FString MetadataPath = FPaths::ChangeExtension(IncomingFileName, TEXT("json"));
+		FFileHelper::SaveStringToFile(Json, *MetadataPath, FFileHelper::EEncodingOptions::ForceUTF8WithoutBOM);
+	}
+
+	// compare the incoming image and throw it away afterwards (note - there will be a copy in the report)
+	TSharedRef<FComparisonEntry> Comparison = MakeShareable(new FComparisonEntry());
+	Comparison->Sender = Context->GetSender();
+	Comparison->ScreenshotPath = Message.Metadata.Context /  Message.Metadata.ScreenShotName;
+	Comparison->PendingComparison = ScreenshotManager->CompareScreenshotAsync(IncomingFileName, Message.Metadata, EScreenShotCompareOptions::DiscardImage);
+
+	ComparisonQueue.Enqueue(Comparison);
+}
+
+void FAutomationControllerManager::HandleTestDataRequest(const FAutomationWorkerTestDataRequest& Message, const TSharedRef<IMessageContext, ESPMode::ThreadSafe>& Context)
+{
+	UE_LOG(LogAutomationController, Log, TEXT("Received TestDataRequest from %s"), *Context->GetSender().ToString());
+
+	const FString TestDataRoot = FPaths::ConvertRelativePathToFull(FPaths::ProjectDir() / TEXT("Test"));
+	const FString DataFile = Message.DataType / Message.DataPlatform / Message.DataTestName / Message.DataName + TEXT(".json");
+	const FString DataFullPath = TestDataRoot / DataFile;
+
+	// Generate the folder for the data if it doesn't exist.
+	const bool bTree = true;
+	IFileManager::Get().MakeDirectory(*FPaths::GetPath(DataFile), bTree);
+
+	bool bIsNew = true;
+	FString ResponseJsonData = Message.JsonData;
+
+	if ( FPaths::FileExists(DataFullPath) )
+	{
+		if ( FFileHelper::LoadFileToString(ResponseJsonData, *DataFullPath) )
+		{
+			bIsNew = false;
+		}
+		else
+		{
+			// TODO Error
+		}
+	}
+
+	if ( bIsNew )
+	{
+
+		FString IncomingTestData = FPaths::Combine(FPaths::AutomationTransientDir(), TEXT("Data/"), DataFile);
+
+		if ( FFileHelper::SaveStringToFile(Message.JsonData, *IncomingTestData) )
+		{
+			//TODO Anything extra to do here?
+		}
+		else
+		{
+			//TODO What do we do if this fails?
+		}
+	}
+
+	FAutomationWorkerTestDataResponse* ResponseMessage = FMessageEndpoint::MakeMessage<FAutomationWorkerTestDataResponse>();
+	ResponseMessage->bIsNew = bIsNew;
+	ResponseMessage->JsonData = ResponseJsonData;
+
+	MessageEndpoint->Send(ResponseMessage, Context->GetSender());
+}
+
+void FAutomationControllerManager::HandlePerformanceDataRequest(const FAutomationWorkerPerformanceDataRequest& Message, const TSharedRef<IMessageContext, ESPMode::ThreadSafe>& Context)
+{
+	//TODO Read/Performance data.
+	UE_LOG(LogAutomationController, Log, TEXT("Received PerformanceDataRequest from %s"), *Context->GetSender().ToString());
+
+	FAutomationWorkerPerformanceDataResponse* ResponseMessage = FMessageEndpoint::MakeMessage<FAutomationWorkerPerformanceDataResponse>();
+	ResponseMessage->bSuccess = true;
+	ResponseMessage->ErrorMessage = TEXT("");
+
+	MessageEndpoint->Send(ResponseMessage, Context->GetSender());
+}
+
+void FAutomationControllerManager::HandleRequestNextNetworkCommandMessage(const FAutomationWorkerRequestNextNetworkCommand& Message, const TSharedRef<IMessageContext, ESPMode::ThreadSafe>& Context)
+{
+	UE_LOG(LogAutomationController, Log, TEXT("Received RequestNextNetworkCommandMessage from %s"), *Context->GetSender().ToString());
+
+	// Harvest iteration of running the tests this result came from (stops stale results from being committed to subsequent runs)
+	if ( Message.ExecutionCount == ExecutionCount )
+	{
+		// Find the device id for the address
+		int32 ClusterIndex;
+		int32 DeviceIndex;
+
+		verify(DeviceClusterManager.FindDevice(Context->GetSender(), ClusterIndex, DeviceIndex));
+
+		// Verify this device thought it was busy
+		TSharedPtr<IAutomationReport> Report = DeviceClusterManager.GetTest(ClusterIndex, DeviceIndex);
+		check(Report.IsValid());
+
+		// Increment network command responses
+		bool bAllResponsesReceived = Report->IncrementNetworkCommandResponses();
+
+		// Test if we've accumulated all responses AND this was the result for the round of test running AND we're still running tests
+		if ( bAllResponsesReceived && ( ClusterDistributionMask & ( 1 << ClusterIndex ) ) )
+		{
+			// Reset the counter
+			Report->ResetNetworkCommandResponses();
+
+			// For every device in this networked test
+			TArray<FMessageAddress> DeviceAddresses = DeviceClusterManager.GetDevicesReservedForTest(ClusterIndex, Report);
+			check(DeviceAddresses.Num() == Report->GetNumParticipantsRequired());
+
+			// Send it to each device
+			for ( int32 AddressIndex = 0; AddressIndex < DeviceAddresses.Num(); ++AddressIndex )
+			{
+				//send "next command message" to worker
+				MessageEndpoint->Send(FMessageEndpoint::MakeMessage<FAutomationWorkerNextNetworkCommandReply>(), DeviceAddresses[AddressIndex]);
+			}
+		}
+	}
+}
+
+void FAutomationControllerManager::HandleRequestTestsReplyCompleteMessage(const FAutomationWorkerRequestTestsReplyComplete& Message, const TSharedRef<IMessageContext, ESPMode::ThreadSafe>& Context)
+{
+	UE_LOG(LogAutomationController, Log, TEXT("Received RequestTestsReplyCompleteMessage from %s"), *Context->GetSender().ToString());
+
+	TArray<FAutomationTestInfo> TestInfo;
+	TestInfo.Reset(Message.Tests.Num());
+	for (const FAutomationWorkerSingleTestReply& SingleTestReply : Message.Tests)
+	{
+		FAutomationTestInfo NewTest = SingleTestReply.GetTestInfo();
+		TestInfo.Add(NewTest);
+	}
+
+	UE_LOG(LogAutomationController, Log, TEXT("%d tests available on %s"), TestInfo.Num(), *Context->GetSender().ToString());
+
+	SetTestNames(Context->GetSender(), TestInfo);
+}
+
+void FAutomationControllerManager::HandleReceivedTelemetryData(const FAutomationWorkerTelemetryData& Message, const TSharedRef<IMessageContext, ESPMode::ThreadSafe>& Context)
+{
+	FAutomationTelemetry::HandleAddTelemetry(Message);
+}
+
+void FAutomationControllerManager::ReportAutomationResult(const TSharedPtr<IAutomationReport> InReport, int32 ClusterIndex, int32 PassIndex)
+{
+
+	FName CategoryName = "LogAutomationController";
+#if WITH_EDITOR
+	FMessageLog AutomationEditorLog("AutomationTestingLog");
+	// we log these messages ourselves for non-editor platforms so suppress this.
+	AutomationEditorLog.SuppressLoggingToOutputLog(true);
+	AutomationEditorLog.Open();
+#endif
+
+	const FAutomationTestResults& Results = InReport->GetResults(ClusterIndex, PassIndex);
+	const FString StateString = AutomationStateToString(Results.State);
+
+	// write results to editor panel
+#if WITH_EDITOR
+	FFormatNamedArguments Arguments;
+	Arguments.Add(TEXT("TestName"), FText::FromString(InReport->GetFullTestPath()));
+	Arguments.Add(TEXT("Result"), FText::FromString(StateString));
+	TSharedRef<FTextToken> Token = FTextToken::Create(FText::Format(LOCTEXT("TestSuccessOrFailure", "Test '{TestName}' completed with result '{Result}'"), Arguments));
+	
+	if (Results.State == EAutomationState::Success)
+	{
+		AutomationEditorLog.Info()->AddToken(Token);
+	}
+	else if (Results.State == EAutomationState::Skipped)
+	{
+		AutomationEditorLog.Warning()->AddToken(Token);
+	}
+	else
+	{
+		AutomationEditorLog.Error()->AddToken(Token);
+	}
+#endif
+
+	// Now log
+	FString ResultString = FString::Printf(AutomationStateFormat, *StateString, *InReport->GetDisplayName(), *InReport->GetFullTestPath());
+	if (Results.State == EAutomationState::Fail)
+	{
+		UE_LOG(LogAutomationController, Error, TEXT("%s"), *ResultString);
+	}
+	else
+	{
+		UE_LOG(LogAutomationController, Display, TEXT("%s"), *ResultString);
+	}
+
+	// bracket these for easy parsing
+	UE_LOG(LogAutomationController, Log, BeginEventsFormat, *InReport->GetFullTestPath());
+
+	for (const FAutomationExecutionEntry& Entry : Results.GetEntries())
+	{
+		switch (Entry.Event.Type)
+		{
+		case EAutomationEventType::Info:
+			UE_LOG(LogAutomationController, Log, TEXT("%s"), *Entry.ToString());
+#if WITH_EDITOR
+			AutomationEditorLog.Info(FText::FromString(Entry.ToString()));
+#endif
+			break;
+		case EAutomationEventType::Warning:
+			UE_LOG(LogAutomationController, Warning, TEXT("%s"), *Entry.ToString());
+#if WITH_EDITOR
+			AutomationEditorLog.Warning(FText::FromString(Entry.ToString()));
+#endif
+			break;
+		case EAutomationEventType::Error:
+			UE_LOG(LogAutomationController, Error, TEXT("%s"), *Entry.ToString());
+#if WITH_EDITOR
+			AutomationEditorLog.Error(FText::FromString(Entry.ToString()));
+#endif
+			break;
+		}
+	}
+
+	UE_LOG(LogAutomationController, Log, EndEventsFormat, *InReport->GetFullTestPath());
+
+	#undef AutomationSuccessFormat
+	#undef AutomationFailureFormat
+	#undef BeginEventsFormat
+	#undef EndEventsFormat
+}
+
+
+void FAutomationControllerManager::HandleRunTestsReplyMessage(const FAutomationWorkerRunTestsReply& Message, const TSharedRef<IMessageContext, ESPMode::ThreadSafe>& Context)
+{
+	// If we should commit these results
+	if ( Message.ExecutionCount == ExecutionCount )
+	{
+		FAutomationTestResults TestResults;
+
+		TestResults.State = Message.State;
+		TestResults.Duration = Message.Duration;
+
+		// Mark device as back on the market
+		int32 ClusterIndex;
+		int32 DeviceIndex;
+
+		verify(DeviceClusterManager.FindDevice(Context->GetSender(), ClusterIndex, DeviceIndex));
+
+		TestResults.GameInstance = DeviceClusterManager.GetClusterDeviceName(ClusterIndex, DeviceIndex);
+		TestResults.SetEvents(Message.Entries, Message.WarningTotal, Message.ErrorTotal);
+
+		// Verify this device thought it was busy
+		TSharedPtr<IAutomationReport> Report = DeviceClusterManager.GetTest(ClusterIndex, DeviceIndex);
+		if (Report.IsValid())
+		{
+			Report->SetResults(ClusterIndex, CurrentTestPass, TestResults);
+
+			const FAutomationTestResults& FinalResults = Report->GetResults(ClusterIndex, CurrentTestPass);
+
+			ReportAutomationResult(Report, ClusterIndex, CurrentTestPass);
+
+			// Gather all of the data relevant to this test for our json reporting.
+			CollectTestResults(Report, FinalResults);
+		}
+
+		// Device is now good to go
+		DeviceClusterManager.SetTest(ClusterIndex, DeviceIndex, NULL);
+	}
+
+	// Remove the running test
+	RemoveTestRunning(Context->GetSender());
+}
+
+void FAutomationControllerManager::HandleWorkerOfflineMessage( const FAutomationWorkerWorkerOffline& Message, const TSharedRef<IMessageContext, ESPMode::ThreadSafe>& Context )
+{
+	FMessageAddress DeviceMessageAddress = Context->GetSender();
+	DeviceClusterManager.Remove(DeviceMessageAddress);
+}
+
+bool FAutomationControllerManager::IsDeviceGroupFlagSet( EAutomationDeviceGroupTypes::Type InDeviceGroup ) const
+{
+	const uint32 FlagMask = 1 << InDeviceGroup;
+	return (DeviceGroupFlags & FlagMask) > 0;
+}
+
+void FAutomationControllerManager::ToggleDeviceGroupFlag( EAutomationDeviceGroupTypes::Type InDeviceGroup )
+{
+	const uint32 FlagMask = 1 << InDeviceGroup;
+	DeviceGroupFlags = DeviceGroupFlags ^ FlagMask;
+}
+
+void FAutomationControllerManager::UpdateDeviceGroups( )
+{
+	DeviceClusterManager.ReGroupDevices( DeviceGroupFlags );
+
+	// Update the reports in case the number of clusters changed
+	int32 NumOfClusters = DeviceClusterManager.GetNumClusters();
+	ReportManager.ClustersUpdated(NumOfClusters);
+}
+
+void FAutomationControllerManager::ResetAutomationTestTimeout(const TCHAR* Reason)
+{
+	//GLog->Logf(ELogVerbosity::Display, TEXT("Resetting automation test timeout: %s"), Reason);
+	LastTimeUpdateTicked = FPlatformTime::Seconds();
+}
+
+#undef LOCTEXT_NAMESPACE
\ No newline at end of file
diff --git a/Engine/Source/Developer/AutomationController/Private/AutomationControllerManager.h b/Engine/Source/Developer/AutomationController/Private/AutomationControllerManager.h
--- a/Engine/Source/Developer/AutomationController/Private/AutomationControllerManager.h
+++ b/Engine/Source/Developer/AutomationController/Private/AutomationControllerManager.h
@@ -278,6 +278,11 @@ public:
 		ReportManager.GetEnabledTestNames(OutEnabledTestNames);
 	}
 
+	virtual void GetFilteredTestNames(TArray<FString>& OutFilteredTestNames) const override
+	{
+		ReportManager.GetFilteredTestNames(OutFilteredTestNames);
+	}
+
 	virtual void SetEnabledTests(const TArray<FString>& EnabledTests) override
 	{
 		ReportManager.SetEnabledTests(EnabledTests);
diff --git a/Engine/Source/Developer/AutomationController/Private/AutomationControllerManger.cpp b/Engine/Source/Developer/AutomationController/Private/AutomationControllerManger.cpp
deleted file mode 100644
diff --git a/Engine/Source/Developer/AutomationController/Private/AutomationReport.cpp b/Engine/Source/Developer/AutomationController/Private/AutomationReport.cpp
--- a/Engine/Source/Developer/AutomationController/Private/AutomationReport.cpp
+++ b/Engine/Source/Developer/AutomationController/Private/AutomationReport.cpp
@@ -136,6 +136,40 @@ void FAutomationReport::GetEnabledTestNames(TArray<FString>& OutEnabledTestNames
 	return;
 }
 
+void FAutomationReport::GetFilteredTestNames(TArray<FString>& OutFilteredTestNames, FString CurrentPath) const
+{
+	// start from FilteredChildReports
+	if (CurrentPath.IsEmpty())
+	{
+		for (int32 FilteredChildIndex = 0; FilteredChildIndex < FilteredChildReports.Num(); ++FilteredChildIndex)
+		{
+			FilteredChildReports[FilteredChildIndex]->GetFilteredTestNames(OutFilteredTestNames, TestInfo.GetDisplayName());
+		}
+	}
+	else // then continue collecting all leaf nodes
+	{
+		//if this is a leaf collect full test name
+		if (FilteredChildReports.Num() == 0)
+		{
+			const FString FullTestName = CurrentPath.Len() > 0 ? CurrentPath.AppendChar(TCHAR('.')) + TestInfo.GetDisplayName() : TestInfo.GetDisplayName();
+			OutFilteredTestNames.Add(FullTestName);
+		}
+		else
+		{
+			if (!CurrentPath.IsEmpty())
+			{
+				CurrentPath += TEXT(".");
+			}
+			CurrentPath += TestInfo.GetDisplayName();
+			//recurse through the hierarchy
+			for (int32 ChildIndex = 0; ChildIndex < FilteredChildReports.Num(); ++ChildIndex)
+			{
+				FilteredChildReports[ChildIndex]->GetFilteredTestNames(OutFilteredTestNames, CurrentPath);
+			}
+		}
+	}
+	return;
+}
 
 void FAutomationReport::SetEnabledTests(const TArray<FString>& InEnabledTests, FString CurrentPath)
 {
@@ -683,7 +717,6 @@ void FAutomationReport::GetEnabledTestReports(TArray<TSharedPtr<IAutomationRepor
 	}
 }
 
-
 const bool FAutomationReport::HasErrors()
 {
 	bool bHasErrors = false;
diff --git a/Engine/Source/Developer/AutomationController/Private/AutomationReport.h b/Engine/Source/Developer/AutomationController/Private/AutomationReport.h
--- a/Engine/Source/Developer/AutomationController/Private/AutomationReport.h
+++ b/Engine/Source/Developer/AutomationController/Private/AutomationReport.h
@@ -32,6 +32,7 @@ public:
 	virtual int32 GetTotalNumFilteredChildren() const override;
 	virtual int32 GetEnabledTestsNum() const override;
 	virtual void GetEnabledTestNames(TArray<FString>& OutEnabledTestNames, FString CurrentPath) const override;
+	virtual void GetFilteredTestNames(TArray<FString>& OutFilteredTestNames, FString CurrentPath) const override;
 	virtual void SetEnabledTests(const TArray<FString>& EnabledTests, FString CurrentPath) override;
 	virtual bool IsEnabled() const override;
 	virtual void SetEnabled(bool bShouldBeEnabled) override;
diff --git a/Engine/Source/Developer/AutomationController/Private/AutomationReportManager.cpp b/Engine/Source/Developer/AutomationController/Private/AutomationReportManager.cpp
--- a/Engine/Source/Developer/AutomationController/Private/AutomationReportManager.cpp
+++ b/Engine/Source/Developer/AutomationController/Private/AutomationReportManager.cpp
@@ -91,6 +91,10 @@ void FAutomationReportManager::GetEnabledTestNames(TArray<FString>& OutEnabledTe
 	ReportRoot->GetEnabledTestNames(OutEnabledTestNames,FString());
 }
 
+void FAutomationReportManager::GetFilteredTestNames(TArray<FString>& OutFilteredTestNames) const
+{
+	ReportRoot->GetFilteredTestNames(OutFilteredTestNames, FString());
+}
 
 void FAutomationReportManager::SetEnabledTests(const TArray<FString>& EnabledTests)
 {
diff --git a/Engine/Source/Developer/AutomationController/Private/AutomationReportManager.h b/Engine/Source/Developer/AutomationController/Private/AutomationReportManager.h
--- a/Engine/Source/Developer/AutomationController/Private/AutomationReportManager.h
+++ b/Engine/Source/Developer/AutomationController/Private/AutomationReportManager.h
@@ -71,6 +71,9 @@ public:
 	/** Gets the names of enabled tests. */
 	void GetEnabledTestNames(TArray<FString>& OutEnabledTestNames) const;
 
+	/** Gets the names of filtered tests. */
+	void GetFilteredTestNames(TArray<FString>& OutFilteredTestNames) const;
+
 	/** Sets the enabled tests based off the passed list of enabled tests. */
 	void SetEnabledTests(const TArray<FString>& EnabledTests);
 
diff --git a/Engine/Source/Developer/AutomationController/Public/AutomationFilter.h b/Engine/Source/Developer/AutomationController/Public/AutomationFilter.h
new file mode 100644
--- /dev/null
+++ b/Engine/Source/Developer/AutomationController/Public/AutomationFilter.h
@@ -0,0 +1,176 @@
+// Copyright Epic Games, Inc. All Rights Reserved.
+
+#pragma once
+
+#include "CoreMinimal.h"
+#include "IAutomationReport.h"
+#include "Misc/IFilter.h"
+
+class FAutomationFilter
+	: public IFilter<const TSharedPtr<class IAutomationReport>&>
+{
+public:
+
+	/** Default constructor. */
+	FAutomationFilter()
+		: OnlySmokeTests( false )
+		, ShowErrors( false )
+		, ShowWarnings( false )
+		, OnlyExcludedTests ( false )
+	{ }
+
+public:
+
+	/**
+	 * Set if we should only show warnings.
+	 *
+	 * @param InShowWarnings If we should show warnings or not.
+	 * @see SetOnlyShowSmokeTests, SetShowErrors, ShouldShowWarnings
+	 */
+	void SetShowWarnings( const bool InShowWarnings )
+	{
+		ShowWarnings = InShowWarnings;
+	}
+
+	/**
+	 * Should we show warnings.
+	 *
+	 * @return True if we should be showing warnings.
+	 * @see SetShowErrors, SetShowWarnings
+	 */
+	const bool ShouldShowWarnings( ) const
+	{
+		return ShowWarnings;
+	}
+
+	/**
+	 * Set if we should only show errors.
+	 *
+	 * @param InShowErrors If we should show errors.
+	 * @see SetOnlyShowSmokeTests, SetShowWarnings, ShouldShowErrors
+	 */
+	void SetShowErrors( const bool InShowErrors )
+	{
+		ShowErrors = InShowErrors;
+	}
+
+	/**
+	 * Should we show errors.
+	 *
+	 * @return True if we should be showing errors.
+	 * @see SetShowErrors, ShouldShowWarnings
+	 */
+	const bool ShouldShowErrors( ) const
+	{
+		return ShowErrors;
+	}
+
+	/**
+	 * Set if we should only show smoke tests.
+	 *
+	 * @param InOnlySmokeTests If we should show smoke tests.
+	 */
+	void SetOnlyShowSmokeTests( const bool InOnlySmokeTests )
+	{
+		OnlySmokeTests = InOnlySmokeTests;
+	}
+
+	/**
+	 * Should we only show errors.
+	 *
+	 * @return True if we should be showing errors.
+	 * @see SetOnlyShowSmokeTests, ShouldShowErrors, ShouldShowWarnings
+	 */
+	const bool OnlyShowSmokeTests( ) const
+	{
+		return OnlySmokeTests;
+	}
+
+	/**
+	 * Should we only show excluded tests.
+	 *
+	 * @return True if we should be showing only excluded tests.
+	 * @see SetOnlyShowSmokeTests, ShouldShowErrors, ShouldShowWarnings
+	 */
+	const bool ShouldShowOnlyExcludedTests() const
+	{
+		return OnlyExcludedTests;
+	}
+
+	/**
+	 * Set if we should only show excluded tests.
+	 *
+	 * @param InOnlyExcludedTests If we should show excluded tests.
+	 */
+	void SetShowOnlyExcludedTests(const bool InOnlyExcludedTests)
+	{
+		OnlyExcludedTests = InOnlyExcludedTests;
+	}
+
+public:
+
+	// IFilter interface
+
+	DECLARE_DERIVED_EVENT(FAutomationFilter, IFilter< const TSharedPtr< class IAutomationReport >& >::FChangedEvent, FChangedEvent);
+	virtual FChangedEvent& OnChanged() override { return ChangedEvent; }
+
+	virtual bool PassesFilter( const TSharedPtr< IAutomationReport >& InReport ) const override
+	{
+		bool FilterPassed = true;
+
+		if (OnlySmokeTests)
+		{
+			//if we only want smoke tests and this isn't one
+			if (!InReport->IsSmokeTest())
+			{
+				return false;
+			}
+			//Leaf Nodes dictate this matching, not root nodes
+			if (InReport->GetTotalNumChildren())
+			{
+				return false;
+			}
+		}
+
+		if (OnlyExcludedTests)
+		{
+			if (!InReport->IsToBeSkipped())
+			{
+				return false;
+			}
+		}
+
+		if (ShowWarnings && ShowErrors)
+		{
+			FilterPassed = InReport->HasWarnings() || InReport->HasErrors();
+		}
+		else if (!ShowWarnings && ShowErrors)
+		{
+			FilterPassed = InReport->HasErrors();
+		}
+		else if (ShowWarnings && !ShowErrors)
+		{
+			// Do not show report as a warning if it should be highlighted as an Error!
+			FilterPassed = InReport->HasWarnings() && !InReport->HasErrors();
+		}
+
+		return FilterPassed;
+	}
+
+private:
+
+	/**	The event that broadcasts whenever a change occurs to the filter. */
+	FChangedEvent ChangedEvent;
+	
+	/** Only smoke tests will pass the test. */
+	bool OnlySmokeTests;
+
+	/** Only errors will pass the test. */
+	bool ShowErrors;
+
+	/** Only warnings will pass the test. */
+	bool ShowWarnings;
+
+	/** Only exclude tests. */
+	bool OnlyExcludedTests;
+};
diff --git a/Engine/Source/Developer/AutomationController/Public/AutomationGroupFilter.h b/Engine/Source/Developer/AutomationController/Public/AutomationGroupFilter.h
new file mode 100644
--- /dev/null
+++ b/Engine/Source/Developer/AutomationController/Public/AutomationGroupFilter.h
@@ -0,0 +1,91 @@
+// Copyright Epic Games, Inc. All Rights Reserved.
+
+#pragma once
+
+#include "CoreMinimal.h"
+#include "IAutomationReport.h"
+#include "Misc/IFilter.h"
+#include "AutomationControllerSettings.h"
+
+class FAutomationGroupFilter
+	: public IFilter<const TSharedPtr<class IAutomationReport>&>
+{
+public:
+
+	/** Default constructor. */
+	FAutomationGroupFilter()
+	{ }
+
+	/** Default constructor with array param. */
+	FAutomationGroupFilter(const TArray<FAutomatedTestFilter> InFilters)
+		: Filters(InFilters)
+	{ }
+
+	/** Default constructor, group with single elements. */
+	FAutomationGroupFilter(const FAutomatedTestFilter InFilter)
+	{
+		Filters.Add(InFilter);
+	}
+
+public:
+	/**
+	 * Set the list of strings the group filter checks for substrings in test display name.
+	 *
+	 * @param InContainsArray An array of strings to filter against test display names.
+	 * @see SetMatchFromStartArray, SetMatchFromEndArray
+	 */
+	void SetFilters(const TArray<FAutomatedTestFilter> InFilters)
+	{
+		Filters = InFilters;
+		ChangedEvent.Broadcast();
+	}
+public:
+
+	// IFilter interface
+
+	DECLARE_DERIVED_EVENT(FAutomationGroupFilter, IFilter< const TSharedPtr< class IAutomationReport >& >::FChangedEvent, FChangedEvent);
+	virtual FChangedEvent& OnChanged() override { return ChangedEvent; }
+
+	virtual bool PassesFilter( const TSharedPtr< IAutomationReport >& InReport ) const override
+	{
+		const int NumGroupFilters = Filters.Num();
+		bool bMeetsMatch;
+
+		for (const FAutomatedTestFilter& Filter: Filters)
+		{
+			bMeetsMatch = true;
+
+			if (Filter.MatchFromStart || Filter.MatchFromEnd)
+			{
+				if (Filter.MatchFromStart)
+				{
+					bMeetsMatch = InReport->GetFullTestPath().StartsWith(Filter.Contains);
+				}
+
+				if (Filter.MatchFromEnd && bMeetsMatch)
+				{
+					bMeetsMatch = InReport->GetFullTestPath().EndsWith(Filter.Contains);
+				}
+			}
+			else
+			{
+				bMeetsMatch = InReport->GetFullTestPath().Contains(Filter.Contains);
+			}
+
+			if (bMeetsMatch)
+			{
+				return true;
+			}
+		}
+		return Filters.Num() == 0;
+	}
+
+private:
+
+	/**	The event that broadcasts whenever a change occurs to the filter. */
+	FChangedEvent ChangedEvent;
+
+	/** The array of FAutomatedTestFilter to filter against. At least one from the list must be matched. */
+	TArray<FAutomatedTestFilter> Filters;
+
+};
diff --git a/Engine/Source/Developer/AutomationController/Public/IAutomationControllerManager.h b/Engine/Source/Developer/AutomationController/Public/IAutomationControllerManager.h
--- a/Engine/Source/Developer/AutomationController/Public/IAutomationControllerManager.h
+++ b/Engine/Source/Developer/AutomationController/Public/IAutomationControllerManager.h
@@ -272,6 +272,13 @@ public:
 	 */
 	virtual void GetEnabledTestNames(TArray<FString>& OutEnabledTestNames) const = 0;
 
+	/**
+	* Gets the names of all the filtered tests.
+	*
+	* @param OutFilteredTestNames The list to populate with filtered test names.
+	*/
+	virtual void GetFilteredTestNames(TArray<FString>& OutFilteredTestNames) const = 0;
+
 	/**
 	 * Sets any tests that match a name in the enabled tests array.
 	 *
diff --git a/Engine/Source/Developer/AutomationController/Public/IAutomationReport.h b/Engine/Source/Developer/AutomationController/Public/IAutomationReport.h
--- a/Engine/Source/Developer/AutomationController/Public/IAutomationReport.h
+++ b/Engine/Source/Developer/AutomationController/Public/IAutomationReport.h
@@ -255,6 +255,9 @@ public:
 	/** Gets the names of all the enabled tests */
 	virtual void GetEnabledTestNames(TArray<FString>& OutEnabledTestNames, FString CurrentPath) const = 0;
 
+	/** Gets the names of all the filtered tests */
+	virtual void GetFilteredTestNames(TArray<FString>& OutFilteredTestNames, FString CurrentPath) const = 0;
+
 	/** Sets which tests are enabled based off the enabled tests list */
 	virtual void SetEnabledTests(const TArray<FString>& EnabledTests, FString CurrentPath) = 0;
 
diff --git a/Engine/Source/Developer/AutomationWindow/Private/AutomationFilter.h b/Engine/Source/Developer/AutomationWindow/Private/AutomationFilter.h
deleted file mode 100644
diff --git a/Engine/Source/Developer/AutomationWindow/Private/SAutomationWindow.cpp b/Engine/Source/Developer/AutomationWindow/Private/SAutomationWindow.cpp
--- a/Engine/Source/Developer/AutomationWindow/Private/SAutomationWindow.cpp
+++ b/Engine/Source/Developer/AutomationWindow/Private/SAutomationWindow.cpp
@@ -24,6 +24,7 @@
 #include "Widgets/Input/SSpinBox.h"
 #include "SAutomationWindowCommandBar.h"
 #include "AutomationFilter.h"
+#include "AutomationGroupFilter.h"
 #include "AutomationPresetManager.h"
 #include "SAutomationTestItemContextMenu.h"
 #include "SAutomationTestItem.h"
@@ -40,6 +41,7 @@
 #include "Widgets/Input/SHyperlink.h"
 #include "Internationalization/Regex.h"
 #include "AutomationWindowStyle.h"
+#include "AutomationControllerSettings.h"
 
 
 #define LOCTEXT_NAMESPACE "AutomationTest"
@@ -171,9 +173,11 @@ void SAutomationWindow::Construct( const FArguments& InArgs, const IAutomationCo
 	// Create the search filter and set criteria
 	AutomationTextFilter = MakeShareable( new AutomationReportTextFilter( AutomationReportTextFilter::FItemToStringArray::CreateSP( this, &SAutomationWindow::PopulateReportSearchStrings ) ) );
 	AutomationGeneralFilter = MakeShareable( new FAutomationFilter() ); 
+	AutomationGroupFilter = MakeShareable( new FAutomationGroupFilter() );
 	AutomationFilters = MakeShareable( new AutomationFilterCollection() );
 	AutomationFilters->Add( AutomationTextFilter );
 	AutomationFilters->Add( AutomationGeneralFilter );
+	AutomationFilters->Add( AutomationGroupFilter );
 
 	bIsRequestingTests = false;
 
@@ -275,6 +279,26 @@ void SAutomationWindow::Construct( const FArguments& InArgs, const IAutomationCo
 	RequestedFilterComboList.Add(MakeShareable(new FString(TEXT("Standard Tests"))));
 	RequestedFilterComboList.Add(MakeShareable(new FString(TEXT("Negative Tests"))));
 
+	UAutomationControllerSettings* Settings = UAutomationControllerSettings::StaticClass()->GetDefaultObject<UAutomationControllerSettings>();
+	TArray<FAutomatedTestGroup> TestGroups = Settings->Groups;
+	GroupComboList.Empty();
+	GroupFiltersMap.Empty();
+	if (TestGroups.Num() > 0)
+	{
+		const FString AllGroups = TEXT("All Groups");
+		GroupComboList.Add(MakeShareable(new FString(AllGroups)));
+		GroupFiltersMap.Add(AllGroups, TArray<FAutomatedTestFilter>());
+		for (int TestGroupIdx = 0; TestGroupIdx < TestGroups.Num(); TestGroupIdx++)
+		{
+			GroupComboList.Add(MakeShareable(new FString(TestGroups[TestGroupIdx].Name)));
+			GroupFiltersMap.Add(TestGroups[TestGroupIdx].Name, TestGroups[TestGroupIdx].Filters);
+		}
+	}
+	else
+	{
+		GroupComboBox->SetVisibility(EVisibility::Collapsed);
+	}
+	
 	TSharedRef<SNotificationList> NotificationList = SNew(SNotificationList) .Visibility( EVisibility::HitTestInvisible );
 
 	//build the actual guts of the window
@@ -351,6 +375,26 @@ void SAutomationWindow::Construct( const FArguments& InArgs, const IAutomationCo
 											]
 										]
 
+										+ SHorizontalBox::Slot()
+										.AutoWidth()
+										.VAlign(VAlign_Center)
+										[
+											SNew(SBox)
+											.MinDesiredWidth(130.0f)
+											[
+												SAssignNew(GroupComboBox, SComboBox< TSharedPtr<FString> >)
+												.OptionsSource(&GroupComboList)
+												.InitiallySelectedItem(GroupComboList[0])
+												.OnGenerateWidget(this, &SAutomationWindow::GenerateGroupComboItem)
+												.OnSelectionChanged(this, &SAutomationWindow::HandleGroupChanged)
+												.ContentPadding(FMargin(4.0, 1.0f))
+												[
+													SNew(STextBlock)
+													.Text(this, &SAutomationWindow::GetGroupComboText)
+												]
+											]
+										]
+
 										+ SHorizontalBox::Slot()
 										.FillWidth(1.0f)
 										.VAlign(VAlign_Center)
@@ -989,6 +1033,11 @@ void SAutomationWindow::HandleRequesteFilterChanged(TSharedPtr<FString> Item, ES
 	AutomationController->SetRequestedTestFlags(NewRequestedFlags);
 }
 
+void SAutomationWindow::HandleGroupChanged(TSharedPtr<FString> Item, ESelectInfo::Type SelectInfo)
+{
+	AutomationGroupFilter->SetFilters(GroupFiltersMap[*Item.Get()]);
+	OnRefreshTestCallback();
+}
 
 void SAutomationWindow::ExpandEnabledTests( TSharedPtr< IAutomationReport > InReport )
 {
@@ -1057,6 +1106,17 @@ FText SAutomationWindow::GetRequestedFilterComboText() const
 	}
 }
 
+FText SAutomationWindow::GetGroupComboText() const
+{
+	if (GroupComboBox->GetSelectedItem().IsValid())
+	{
+		return FText::FromString(*GroupComboBox->GetSelectedItem());
+	}
+	else
+	{
+		return LOCTEXT("AutomationGroupComboLabel", "All Groups");
+	}
+}
 
 TSharedRef<SWidget> SAutomationWindow::GeneratePresetComboItem(TSharedPtr<FAutomationTestPreset> InItem)
 {
@@ -1078,6 +1138,12 @@ TSharedRef<SWidget> SAutomationWindow::GenerateRequestedFilterComboItem(TSharedP
 		.Text(FText::FromString(*InItem));
 }
 
+TSharedRef<SWidget> SAutomationWindow::GenerateGroupComboItem(TSharedPtr<FString> InItem)
+{
+	return SNew(STextBlock)
+		.Text(FText::FromString(*InItem));
+}
+
 TSharedRef< SWidget > SAutomationWindow::GenerateGroupOptionsMenuContent( TWeakPtr<class SAutomationWindow> InAutomationWindow )
 {
 	TSharedPtr<SAutomationWindow> AutomationWindow(InAutomationWindow.Pin());
diff --git a/Engine/Source/Developer/AutomationWindow/Private/SAutomationWindow.h b/Engine/Source/Developer/AutomationWindow/Private/SAutomationWindow.h
--- a/Engine/Source/Developer/AutomationWindow/Private/SAutomationWindow.h
+++ b/Engine/Source/Developer/AutomationWindow/Private/SAutomationWindow.h
@@ -22,6 +22,8 @@
 #endif
 
 class FAutomationFilter;
+class FAutomationGroupFilter;
+struct FAutomatedTestFilter;
 class FAutomationTestPresetManager;
 class FUICommandList;
 class SAutomationWindowCommandBar;
@@ -221,6 +223,13 @@ private:
 	 * @return New combo item widget.
 	 */
 	TSharedRef<SWidget> GenerateRequestedFilterComboItem(TSharedPtr<FString> InItem);
+
+	/**
+	 * Creates a combo item for a test group.
+	 *
+	 * @return New combo item widget.
+	 */
+	TSharedRef<SWidget> GenerateGroupComboItem(TSharedPtr<FString> InItem);
 		
 	/**
 	 * Populates OutSearchStrings with the strings that should be used in searching.
@@ -470,6 +479,8 @@ private:
 	void HandlePresetChanged( TSharedPtr<FAutomationTestPreset> Item, ESelectInfo::Type SelectInfo );
 	/** Called when the user changes the requested test filter */
 	void HandleRequesteFilterChanged(TSharedPtr<FString> Item, ESelectInfo::Type SelectInfo);
+	/** Called when the user changes the test group */
+	void HandleGroupChanged(TSharedPtr<FString> Item, ESelectInfo::Type SelectInfo);
 
 	/** Expands the test tree to show all enabled tests. */
 	void ExpandEnabledTests( TSharedPtr< IAutomationReport > InReport );
@@ -478,6 +489,8 @@ private:
 	FText GetPresetComboText() const;
 	/** Gets the text to display for the requested filter combo box. */
 	FText GetRequestedFilterComboText() const;
+	/** Gets the combo box text to display for the selected group. */
+	FText GetGroupComboText() const;
 
 	/**
 	 * Handle the copy button clicked in the command bar.
@@ -577,6 +590,9 @@ private:
 	/** The automation general filter - for smoke tests / warnings and Errors. */
 	TSharedPtr< FAutomationFilter > AutomationGeneralFilter;
 
+	/** The automation group filter - filters tests based on selected group filter. */
+	TSharedPtr< FAutomationGroupFilter > AutomationGroupFilter;
+
 	/** The automation filter collection - contains the automation filters. */
 	TSharedPtr< AutomationFilterCollection > AutomationFilters;
 
@@ -614,6 +630,13 @@ private:
 	TSharedPtr< SComboBox< TSharedPtr<FString> > >	RequestedFilterComboBox;
 	TArray< TSharedPtr< FString > >					RequestedFilterComboList;
 
+	/** Holds a pointer to the groups combo box widget. */
+	TSharedPtr< SComboBox< TSharedPtr<FString> > >	GroupComboBox;
+	TArray< TSharedPtr< FString > >					GroupComboList;
+
+	/** Map for fast access of test group filters by name. */
+	TMap< FString, TArray<FAutomatedTestFilter> > GroupFiltersMap;
+
 	/** Holds a pointer to the preset text box. */
 	TSharedPtr<SEditableTextBox> PresetTextBox;
 
