Tag allocations related to Automation Test components

diff --git a/Engine/Source/Developer/AutomationController/Private/AutomationControllerManager.cpp b/Engine/Source/Developer/AutomationController/Private/AutomationControllerManager.cpp
--- a/Engine/Source/Developer/AutomationController/Private/AutomationControllerManager.cpp
+++ b/Engine/Source/Developer/AutomationController/Private/AutomationControllerManager.cpp
@@ -46,6 +46,7 @@ DEFINE_LOG_CATEGORY_STATIC(LogAutomationController, Log, All)
 
 void FAutomatedTestPassResults::AddTestResult(const IAutomationReportPtr& TestReport)
 {
+	LLM_SCOPE_BYNAME(TEXT("AutomationTest/Report"));
 	const FString& FullTestPath = TestReport->GetFullTestPath();
 	if (!TestsMapIndex.Contains(FullTestPath))
 	{
@@ -69,6 +70,7 @@ FAutomatedTestResult& FAutomatedTestPassResults::GetTestResult(const IAutomation
 
 void FAutomatedTestPassResults::ReBuildTestsMapIndex()
 {
+	LLM_SCOPE_BYNAME(TEXT("AutomationTest/Report"));
 	TestsMapIndex.Empty();
 	uint32 Index = 0;
 	for (const FAutomatedTestResult& TestResult : Tests)
@@ -654,6 +656,7 @@ bool FAutomationControllerManager::GenerateTestPassHtmlIndex()
 
 bool FAutomationControllerManager::LoadJsonTestPassSummary(FString& ReportFilePath, TArray<IAutomationReportPtr> TestReports)
 {
+	LLM_SCOPE_BYNAME(TEXT("AutomationTest/Report"));
 	FString Json;
 	if (!FFileHelper::LoadFileToString(Json, *ReportFilePath))
 	{
@@ -1085,11 +1088,7 @@ void FAutomationControllerManager::UpdateTests()
 					//verify this device thought it was busy
 					TSharedPtr <IAutomationReport> Report = DeviceClusterManager.GetTest(ClusterIndex, DeviceIndex);
 					check(Report.IsValid());
-					// A dummy array used to report the result
 
-					TArray<FString> EmptyStringArray;
-					TArray<FString> ErrorStringArray;
-					ErrorStringArray.Add(FString(TEXT("Failed")));
 					bHasErrors = true;
 
 					FAutomationTestResults TestResults;
@@ -1362,6 +1361,7 @@ void FAutomationControllerManager::HandleRequestTestsReplyCompleteMessage(const
 {
 	UE_LOG(LogAutomationController, Log, TEXT("Received RequestTestsReplyCompleteMessage from %s"), *Context->GetSender().ToString());
 
+	LLM_SCOPE_BYNAME(TEXT("AutomationTest/Controller"));
 	TArray<FAutomationTestInfo> TestInfo;
 	TestInfo.Reset(Message.Tests.Num());
 	for (const FAutomationWorkerSingleTestReply& SingleTestReply : Message.Tests)
diff --git a/Engine/Source/Developer/AutomationController/Private/AutomationReport.cpp b/Engine/Source/Developer/AutomationController/Private/AutomationReport.cpp
--- a/Engine/Source/Developer/AutomationController/Private/AutomationReport.cpp
+++ b/Engine/Source/Developer/AutomationController/Private/AutomationReport.cpp
@@ -114,6 +114,7 @@ int32 FAutomationReport::GetTotalNumFilteredChildren() const
 
 void FAutomationReport::GetEnabledTestNames(TArray<FString>& OutEnabledTestNames, FString CurrentPath) const
 {
+	LLM_SCOPE_BYNAME(TEXT("AutomationTest/Report"));
 	//if this is a leaf and this test is enabled
 	if ((ChildReports.Num() == 0) && IsEnabled())
 	{
@@ -138,6 +139,7 @@ void FAutomationReport::GetEnabledTestNames(TArray<FString>& OutEnabledTestNames
 
 void FAutomationReport::GetFilteredTestNames(TArray<FString>& OutFilteredTestNames, FString CurrentPath) const
 {
+	LLM_SCOPE_BYNAME(TEXT("AutomationTest/Report"));
 	// start from FilteredChildReports
 	if (CurrentPath.IsEmpty())
 	{
@@ -240,6 +242,7 @@ void FAutomationReport::SetSupport(const int32 ClusterIndex)
 	SupportFlags |= (1<<ClusterIndex);
 
 	//ensure there is enough room in the array for status per platform
+	LLM_SCOPE_BYNAME(TEXT("AutomationTest/Report"));
 	for (int32 i = 0; i <= ClusterIndex; ++i)
 	{
 		//Make sure we have enough results for a single pass
@@ -292,6 +295,7 @@ const bool FAutomationReport::IsSmokeTest( )
 
 bool FAutomationReport::SetFilter( TSharedPtr< AutomationFilterCollection > InFilter, const bool ParentPassedFilter )
 {
+	LLM_SCOPE_BYNAME(TEXT("AutomationTest/Report"));
 	//assume that this node and all its children fail to pass the filter test
 	bool bSelfOrChildPassedFilter = false;
 
@@ -351,6 +355,7 @@ TArray<TSharedPtr<IAutomationReport> >& FAutomationReport::GetChildReports()
 
 void FAutomationReport::ClustersUpdated(const int32 NumClusters)
 {
+	LLM_SCOPE_BYNAME(TEXT("AutomationTest/Report"));
 	TestInfo.ResetNumDevicesRunningTest();
 
 	//Fixup Support flags
@@ -431,6 +436,7 @@ void FAutomationReport::SetResults( const int32 ClusterIndex, const int32 PassIn
 		TestInfo.InformOfNewDeviceRunningTest();
 	}
 
+	LLM_SCOPE_BYNAME(TEXT("AutomationTest/Report"));
 	const TArray<FAutomationArtifact> ExistingArtifacts = Results[ClusterIndex][PassIndex].Artifacts;
 	Results[ClusterIndex][PassIndex] = InResults;
 	Results[ClusterIndex][PassIndex].Artifacts.Append(ExistingArtifacts);
@@ -452,6 +458,7 @@ void FAutomationReport::AddArtifact(const int32 ClusterIndex, const int32 PassIn
 	check(( ClusterIndex >= 0 ) && ( ClusterIndex < Results.Num() ));
 	check(( PassIndex >= 0 ) && ( PassIndex < Results[ClusterIndex].Num() ));
 
+	LLM_SCOPE_BYNAME(TEXT("AutomationTest/Report"));
 	Results[ClusterIndex][PassIndex].Artifacts.Add(Artifact);
 }
 
@@ -603,6 +610,7 @@ TSharedPtr<IAutomationReport> FAutomationReport::EnsureReportExists(FAutomationT
 	//if there isn't already a test like this
 	if (!MatchTest.IsValid())
 	{
+		LLM_SCOPE_BYNAME(TEXT("AutomationTest/Report"));
 		if ( NameRemainder.Len() == 0 )
 		{
 			// Create a new leaf node
diff --git a/Engine/Source/Developer/ScreenShotComparisonTools/Private/ImageComparer.cpp b/Engine/Source/Developer/ScreenShotComparisonTools/Private/ImageComparer.cpp
--- a/Engine/Source/Developer/ScreenShotComparisonTools/Private/ImageComparer.cpp
+++ b/Engine/Source/Developer/ScreenShotComparisonTools/Private/ImageComparer.cpp
@@ -6,6 +6,7 @@
 #include "IImageWrapper.h"
 #include "IImageWrapperModule.h"
 #include "ImageWrapperHelper.h"
+#include "HAL/LowLevelMemTracker.h"
 #include "Misc/FileHelper.h"
 #include "Misc/Paths.h"
 #include "Modules/ModuleManager.h"
@@ -214,6 +215,7 @@ FComparisonReport::FComparisonReport(const FString& InReportRootDirectory, const
 
 TSharedPtr<FComparableImage> FImageComparer::Open(const FString& ImagePath, FText& OutError)
 {
+	LLM_SCOPE_BYNAME(TEXT("AutomationTest/ImageCompare"));
 	const FString ImageExtension = FPaths::GetExtension(ImagePath);
 	const EImageFormat ImageFormat = ImageWrapperHelper::GetImageFormat(ImageExtension);
 
@@ -289,6 +291,7 @@ FImageComparisonResult FImageComparer::Compare(const FString& ImagePathA, const
 
 FImageComparisonResult FImageComparer::Compare(const FComparableImage* ImageA, const FComparableImage* ImageB, FImageTolerance Tolerance, const FString& OutDeltaPath)
 {
+	LLM_SCOPE_BYNAME(TEXT("AutomationTest/ImageCompare"));
 	FImageComparisonResult Results;
 
 	// Compare the smallest shared dimensions, this will be a forced failure
diff --git a/Engine/Source/Runtime/AutomationWorker/Private/AutomationWorkerModule.cpp b/Engine/Source/Runtime/AutomationWorker/Private/AutomationWorkerModule.cpp
--- a/Engine/Source/Runtime/AutomationWorker/Private/AutomationWorkerModule.cpp
+++ b/Engine/Source/Runtime/AutomationWorker/Private/AutomationWorkerModule.cpp
@@ -177,42 +177,6 @@ void FAutomationWorkerModule::ReportNetworkCommandComplete()
 	}
 }
 
-
-/**
- *	Takes a large transport array and splits it into pieces of a desired size and returns the portion of this which is requested
- *
- * @param FullTransportArray The whole series of data
- * @param CurrentChunkIndex The The chunk we are requesting
- * @param NumToSend The maximum number of bytes we should be splitting into.
- * @return The section of the transport array which matches our index requested
- */
-TArray< uint8 > GetTransportSection( const TArray< uint8 >& FullTransportArray, const int32 NumToSend, const int32 RequestedChunkIndex )
-{
-	TArray< uint8 > TransportArray = FullTransportArray;
-
-	if( NumToSend > 0 )
-	{
-		int32 NumToRemoveFromStart = RequestedChunkIndex * NumToSend;
-		if( NumToRemoveFromStart > 0 )
-		{
-			TransportArray.RemoveAt( 0, NumToRemoveFromStart );
-		}
-
-		int32 NumToRemoveFromEnd = FullTransportArray.Num() - NumToRemoveFromStart - NumToSend;
-		if( NumToRemoveFromEnd > 0 )
-		{
-			TransportArray.RemoveAt( TransportArray.Num()-NumToRemoveFromEnd, NumToRemoveFromEnd );
-		}
-	}
-	else
-	{
-		TransportArray.Empty();
-	}
-
-	return TransportArray;
-}
-
-
 void FAutomationWorkerModule::ReportTestComplete()
 {
 	if (GIsAutomationTesting)
@@ -275,6 +239,7 @@ void FAutomationWorkerModule::ReportTestComplete()
 
 void FAutomationWorkerModule::SendTests( const FMessageAddress& ControllerAddress )
 {
+	LLM_SCOPE_BYNAME(TEXT("AutomationTest/Worker"));
 	FAutomationWorkerRequestTestsReplyComplete* Reply = FMessageEndpoint::MakeMessage<FAutomationWorkerRequestTestsReplyComplete>();
 	for( int32 TestIndex = 0; TestIndex < TestInfo.Num(); TestIndex++ )
 	{
@@ -450,6 +415,7 @@ void FAutomationWorkerModule::HandleScreenShotCapturedWithName(const TArray<FCol
 void FAutomationWorkerModule::HandleScreenShotAndTraceCapturedWithName(const TArray<FColor>& RawImageData, const TArray<uint8>& CapturedFrameTrace, const FAutomationScreenshotData& Data)
 {
 #if WITH_AUTOMATION_TESTS
+	LLM_SCOPE_BYNAME(TEXT("AutomationTest/ImageCompare"));
 	int32 NewHeight = Data.Height;
 	int32 NewWidth = Data.Width;
 
@@ -524,6 +490,7 @@ void FAutomationWorkerModule::HandleScreenShotAndTraceCapturedWithName(const TAr
 void FAutomationWorkerModule::HandleRunTestsMessage( const FAutomationWorkerRunTests& Message, const TSharedRef<IMessageContext, ESPMode::ThreadSafe>& Context )
 {
 	UE_LOG(LogAutomationWorker, Log, TEXT("Received RunTests %s from %s"), *Message.BeautifiedTestName, *Context->GetSender().ToString());
+	LLM_SCOPE_BYNAME(TEXT("AutomationTest/Worker"));
 
 	if (TestRequesterAddress.IsValid() && !TestName.IsEmpty())
 	{
diff --git a/Engine/Source/Runtime/Core/Private/Misc/AutomationTest.cpp b/Engine/Source/Runtime/Core/Private/Misc/AutomationTest.cpp
--- a/Engine/Source/Runtime/Core/Private/Misc/AutomationTest.cpp
+++ b/Engine/Source/Runtime/Core/Private/Misc/AutomationTest.cpp
@@ -20,7 +20,7 @@ DEFINE_LOG_CATEGORY_STATIC(LogAutomationTest, Warning, All);
 
 CORE_API const TMap<FString, EAutomationTestFlags::Type>& EAutomationTestFlags::GetTestFlagsMap()
 {
-	LLM_SCOPE_BYNAME(TEXT("AutomationTest"));
+	LLM_SCOPE_BYNAME(TEXT("AutomationTest/Framework"));
 	/** String to EAutomationTestFlags map */
 	static const TMap<FString, Type> FlagsMap = {
 		{ TEXT("EditorContext"), Type::EditorContext},
@@ -219,6 +219,7 @@ bool FAutomationTestFramework::RegisterAutomationTest( const FString& InTestName
 	const bool bAlreadyRegistered = AutomationTestClassNameToInstanceMap.Contains( InTestNameToRegister );
 	if ( !bAlreadyRegistered )
 	{
+		LLM_SCOPE_BYNAME(TEXT("AutomationTest/Framework"));
 		AutomationTestClassNameToInstanceMap.Add( InTestNameToRegister, InTestToRegister );
 	}
 	return !bAlreadyRegistered;
@@ -563,6 +564,7 @@ void FAutomationTestFramework::LoadTestModules( )
 
 void FAutomationTestFramework::GetValidTestNames( TArray<FAutomationTestInfo>& TestInfo ) const
 {
+	LLM_SCOPE_BYNAME(TEXT("AutomationTest/Framework"));
 	TestInfo.Empty();
 
 	// Determine required application type (Editor, Game, or Commandlet)
@@ -970,6 +972,8 @@ int32 FAutomationTestExecutionInfo::RemoveAllEvents(TFunctionRef<bool(FAutomatio
 
 void FAutomationTestExecutionInfo::AddEvent(const FAutomationEvent& Event, int StackOffset, bool bCaptureStack)
 {
+	LLM_SCOPE_BYNAME(TEXT("AutomationTest/Framework"));
+
 	switch (Event.Type)
 	{
 	case EAutomationEventType::Warning:
diff --git a/Engine/Source/Runtime/Core/Public/Misc/AutomationTest.h b/Engine/Source/Runtime/Core/Public/Misc/AutomationTest.h
--- a/Engine/Source/Runtime/Core/Public/Misc/AutomationTest.h
+++ b/Engine/Source/Runtime/Core/Public/Misc/AutomationTest.h
@@ -16,6 +16,7 @@
 #include "HAL/PlatformTime.h"
 #include "HAL/ThreadSafeBool.h" 
 #include "HAL/PlatformStackWalk.h"
+#include "HAL/LowLevelMemTracker.h"
 #include "GenericPlatform/GenericPlatformStackWalk.h"
 #include "Containers/Queue.h"
 #include "Misc/CString.h"
@@ -1202,8 +1203,9 @@ public:
 	 */
 	FAutomationTestBase( const FString& InName, const bool bInComplexTask )
 		: bComplexTask( bInComplexTask )
-		, TestName( InName )
 	{
+		LLM_SCOPE_BYNAME(TEXT("AutomationTest/Framework"));
+		TestName = InName;
 		// Register the newly created automation test into the automation testing framework
 		FAutomationTestFramework::Get().RegisterAutomationTest( InName, this );
 	}
@@ -2436,6 +2438,7 @@ public:
 
 	void Describe(const FString& InDescription, TFunction<void()> DoWork)
 	{
+		LLM_SCOPE_BYNAME(TEXT("AutomationTest/Framework"));
 		const TSharedRef<FSpecDefinitionScope> ParentScope = DefinitionScopeStack.Last();
 		const TSharedRef<FSpecDefinitionScope> NewScope = MakeShareable(new FSpecDefinitionScope());
 		NewScope->Description = InDescription;
@@ -2810,6 +2813,7 @@ private:
 
 	void PushDescription(const FString& InDescription)
 	{
+		LLM_SCOPE_BYNAME(TEXT("AutomationTest/Framework"));
 		Description.Add(InDescription);
 	}
 
