Automation command line: filter any instead of all when using multiple filters
#rnx

diff --git a/Engine/Source/Developer/AutomationController/Private/AutomationCommandline.cpp b/Engine/Source/Developer/AutomationController/Private/AutomationCommandline.cpp
--- a/Engine/Source/Developer/AutomationController/Private/AutomationCommandline.cpp
+++ b/Engine/Source/Developer/AutomationController/Private/AutomationCommandline.cpp
@@ -140,6 +140,8 @@ public:
 		// 1) If argument is a filter (filter:system) then make sure we only filter-in tests that start with that filter
 		// 2) If argument is a group then expand that group into multiple filters based on ini entries
 		// 3) Otherwise just substring match (default behavior in 4.22 and earlier).
+		FAutomationGroupFilter* FilterAny = new FAutomationGroupFilter();
+		TArray<FAutomatedTestFilter> FiltersList;
 		for (int32 ArgumentIndex = 0; ArgumentIndex < ArgumentNames.Num(); ++ArgumentIndex)
 		{
 			const FString GroupPrefix = TEXT("Group:");
@@ -157,8 +159,7 @@ public:
 					FilterName += TEXT(".");
 				}
 
-				FAutomationGroupFilter* Filter = new FAutomationGroupFilter(FAutomatedTestFilter(FilterName, true, false));
-				InFilters->Add(MakeShareable(Filter));
+				FiltersList.Add(FAutomatedTestFilter(FilterName, true, false));
 			}
 			else if (ArgumentName.StartsWith(GroupPrefix))
 			{
@@ -176,8 +177,7 @@ public:
 						// if found add all this groups filters to our current list
 						if (GroupEntry->Filters.Num() > 0)
 						{
-							FAutomationGroupFilter* Filter = new FAutomationGroupFilter(GroupEntry->Filters);
-							InFilters->Add(MakeShareable(Filter));
+							FiltersList.Append(GroupEntry->Filters);
 						}
 						else
 						{
@@ -207,11 +207,13 @@ public:
 					ArgumentName.LeftChopInline(1);
 				}
 
-				FAutomationGroupFilter* Filter = new FAutomationGroupFilter(FAutomatedTestFilter(ArgumentName, bMatchFromStart, bMatchFromEnd));
-				InFilters->Add(MakeShareable(Filter));
+				FiltersList.Add(FAutomatedTestFilter(ArgumentName, bMatchFromStart, bMatchFromEnd));
 			}
 		}
 
+		FilterAny->SetFilters(FiltersList);
+		InFilters->Add(MakeShareable(FilterAny));
+
 		// SetFilter applies all filters from the AutomationFilters array
 		AutomationController->SetFilter(InFilters);
 		// Fill OutFilteredTestNames array with filtered test names
