Log Test report path only if it was generated.

diff --git a/Engine/Source/Developer/AutomationController/Private/AutomationControllerManager.cpp b/Engine/Source/Developer/AutomationController/Private/AutomationControllerManager.cpp
--- a/Engine/Source/Developer/AutomationController/Private/AutomationControllerManager.cpp
+++ b/Engine/Source/Developer/AutomationController/Private/AutomationControllerManager.cpp
@@ -984,9 +984,9 @@ void FAutomationControllerManager::ProcessResults()
 				FPlatformProcess::LaunchURL(*DeveloperReportUrl, nullptr, nullptr);
 			}
 		}
-	}
 
-	UE_LOG(LogAutomationController, Display, TEXT("Report can be opened in the editor at '%s'"), ReportExportPath.IsEmpty() ? *FPaths::ConvertRelativePathToFull(FPaths::AutomationReportsDir()) : *ReportExportPath);
+		UE_LOG(LogAutomationController, Display, TEXT("Report can be opened in the editor at '%s'"), ReportExportPath.IsEmpty() ? *FPaths::ConvertRelativePathToFull(FPaths::AutomationReportsDir()) : *ReportExportPath);
+	}
 
 	// Then clean our array for the next pass.
 	JsonTestPassResults.ClearAllEntries();
