Fix test number sum of sub category in Test Automation window
+ fix internal test sub category names in Test Automation window

#rnx

diff --git a/Engine/Source/Developer/AutomationController/Private/AutomationReport.cpp b/Engine/Source/Developer/AutomationController/Private/AutomationReport.cpp
--- a/Engine/Source/Developer/AutomationController/Private/AutomationReport.cpp
+++ b/Engine/Source/Developer/AutomationController/Private/AutomationReport.cpp
@@ -73,7 +73,7 @@ FString FAutomationReport::GetDisplayNameWithDecoration() const
 	//if this is an internal leaf node and the "decoration" name is being requested
 	if (ChildReports.Num())
 	{
-		int32 NumChildren = GetTotalNumChildren();
+		int32 NumChildren = GetTotalNumFilteredChildren();
 		//append on the number of child tests
 		return TestInfo.GetDisplayName() + FString::Printf(TEXT(" (%d)"), NumChildren);
 	}
@@ -580,15 +580,10 @@ TSharedPtr<IAutomationReport> FAutomationReport::EnsureReportExists(FAutomationT
 		// Set the test info name to be the remaining string
 		InTestInfo.SetDisplayName( NameRemainder );
 		// Update the fullpath
-		int32 Pos = FullPath.Find(NameToMatch + TEXT("."));
-		if (Pos >= 0)
-		{
-			FullPath.LeftChopInline(FullPath.Len() - Pos);
-		}
-		FullPath += NameToMatch;
+		FullPath.LeftChopInline(NameRemainder.Len() + 1);
 	}
 
-	uint32 NameToMatchHash = GetTypeHash(NameToMatch);
+	uint32 NameToMatchHash = GetTypeHash(FullPath);
 
 	TSharedPtr<IAutomationReport> MatchTest;
 	//check hash table first to see if it exists yet
@@ -599,7 +594,7 @@ TSharedPtr<IAutomationReport> FAutomationReport::EnsureReportExists(FAutomationT
 		for (; TestIndex >= 0; --TestIndex)
 		{
 			//if the name matches
-			if (ChildReports[TestIndex]->GetDisplayName() == NameToMatch)
+			if (ChildReports[TestIndex]->GetFullTestPath() == FullPath)
 			{
 				MatchTest = ChildReports[TestIndex];
 				break;
