Fix for automation changes in shipping editor builds

diff --git a/Engine/Source/Developer/AutomationController/Private/AutomationControllerManager.cpp b/Engine/Source/Developer/AutomationController/Private/AutomationControllerManager.cpp
--- a/Engine/Source/Developer/AutomationController/Private/AutomationControllerManager.cpp
+++ b/Engine/Source/Developer/AutomationController/Private/AutomationControllerManager.cpp
@@ -211,7 +211,7 @@ bool FAutomationControllerManager::IsReadyForTests()
 		return false;
 	}
 
-#if WITH_EDITOR
+#if WITH_EDITOR && !UE_BUILD_SHIPPING
 	if (InteractiveFrameRateCheck == nullptr)
 	{
 		InteractiveFrameRateCheck = MakeShared<FWaitForInteractiveFrameRate>();
diff --git a/Engine/Source/Developer/AutomationController/Private/AutomationControllerManager.h b/Engine/Source/Developer/AutomationController/Private/AutomationControllerManager.h
--- a/Engine/Source/Developer/AutomationController/Private/AutomationControllerManager.h
+++ b/Engine/Source/Developer/AutomationController/Private/AutomationControllerManager.h
@@ -605,7 +605,7 @@ private:
 
 	bool bResumeRunTest;
 
-#if WITH_EDITOR
+#if WITH_EDITOR && !UE_BUILD_SHIPPING
 	TSharedPtr<class FWaitForInteractiveFrameRate> InteractiveFrameRateCheck;
 #endif
 
