Support individual test exclusion for different RHI when parent suite does not exclude all RHI
+ enhanced management of RHI through the UI: If a RHI is specify, any child test can be excluded on an additional RHI by clicking on edit exclusion
+ disable exclusion editing on remote session

diff --git a/Engine/Source/Developer/AutomationController/Private/AutomationReport.cpp b/Engine/Source/Developer/AutomationController/Private/AutomationReport.cpp
--- a/Engine/Source/Developer/AutomationController/Private/AutomationReport.cpp
+++ b/Engine/Source/Developer/AutomationController/Private/AutomationReport.cpp
@@ -71,7 +71,7 @@ FString FAutomationReport::GetDisplayNameWithDecoration() const
 {
 	FString FinalDisplayName = TestInfo.GetDisplayName();
 	//if this is an internal leaf node and the "decoration" name is being requested
-	if (ChildReports.Num())
+	if (FilteredChildReports.Num())
 	{
 		int32 NumChildren = GetTotalNumFilteredChildren();
 		//append on the number of child tests
@@ -885,6 +885,11 @@ bool FAutomationReport::IsToBeSkipped(FName* OutReason, bool* OutWarn) const
 	return false;
 }
 
+bool FAutomationReport::IsToBeSkippedOnConditions() const
+{
+	return bNeedToSkip && ExcludeTestInfo.HasConditions();
+}
+
 bool FAutomationReport::IsToBeSkippedByPropagation() const
 {
 	return bNeedToSkip && ExcludeTestInfo.bIsPropagated;
@@ -904,7 +909,7 @@ void FAutomationReport::SetSkipFlag(bool bEnableSkip, const FAutomationTestExclu
 		}
 	}
 
-	if (!bFromPropagation && !ExcludeTestInfo.IsEmpty() && ExcludeTestInfo.bIsPropagated)
+	if (!bFromPropagation && !ExcludeTestInfo.IsEmpty() && ExcludeTestInfo.bIsPropagated && !ExcludeTestInfo.HasConditions())
 		return; // Propagated exclusion can't be changed directly
 
 	bNeedToSkip = bEnableSkip;
@@ -916,27 +921,64 @@ void FAutomationReport::SetSkipFlag(bool bEnableSkip, const FAutomationTestExclu
 		ExcludeTestInfo.bIsPropagated = bFromPropagation;
 	}
 
+	auto ExcludedTestCached = UAutomationTestExcludelist::Get();
 	if (bFromPropagation)
 	{
-		if (!bNeedToSkip)
+		auto Entry = ExcludedTestCached->GetExcludeTestEntry(TestInfo.GetFullTestPath());
+		if (bNeedToSkip)
 		{
-			ExcludeTestInfo.Reset();
+			// If we get an exclusion entry, check if it is the original one.
+			if (Entry != nullptr)
+			{
+				ExcludeTestInfo = *Entry;
+				ExcludeTestInfo.SetPropagation(TestInfo.GetFullTestPath());
+			}
+		}
+		else
+		{
+			// Before enabling the test, check if there is an underlying exclusion already set
+			if (Entry == nullptr)
+			{
+				ExcludeTestInfo.Reset();
+			}
+			else
+			{
+				// Update instance exclusion info
+				ExcludeTestInfo = *Entry;
+				bNeedToSkip = true;
+				// Update exclusion template for the children propagation
+				Template = Entry;
+			}
 		}
 	}
 	else
 	{
-		auto ExcludedTestCached = UAutomationTestExcludelist::Get();
 		if (bNeedToSkip)
 		{
 			check(Template != nullptr);
+			auto Entry = ExcludedTestCached->GetExcludeTestEntry(TestInfo.GetFullTestPath());
+			if (Entry != nullptr && ExcludeTestInfo.RemoveConditions(*Entry))
+			{
+				// Branch off the template with a different exclusion condition set
+				Template = &ExcludeTestInfo;
+			}
 			ExcludedTestCached->AddToExcludeTest(TestInfo.GetFullTestPath(), *Template);
 		}
 		else
 		{
 			ExcludedTestCached->RemoveFromExcludeTest(TestInfo.GetFullTestPath());
+			auto Entry = ExcludedTestCached->GetExcludeTestEntry(TestInfo.GetFullTestPath());
+			if (Entry != nullptr)
+			{
+				// If there is still an entry, it means a higher exclusion rule exists so we apply it.
+				// Update instance exclusion info
+				ExcludeTestInfo = *Entry;
+				ExcludeTestInfo.bIsPropagated = true;
+				bNeedToSkip = true;
+				// Update exclusion template for the children propagation
+				Template = Entry;
+			}
 		}
-
-		ExcludedTestCached->SaveConfig();
 	}
 
 	// Propagate to children
@@ -948,6 +990,11 @@ void FAutomationReport::SetSkipFlag(bool bEnableSkip, const FAutomationTestExclu
 		}
 	}
 
+	if (!bFromPropagation)
+	{
+		// Save config only at the end of the recursion
+		ExcludedTestCached->SaveConfig();
+	}
 }
 
 TSharedPtr<FAutomationTestExcludeOptions> FAutomationReport::GetExcludeOptions()
diff --git a/Engine/Source/Developer/AutomationController/Private/AutomationReport.h b/Engine/Source/Developer/AutomationController/Private/AutomationReport.h
--- a/Engine/Source/Developer/AutomationController/Private/AutomationReport.h
+++ b/Engine/Source/Developer/AutomationController/Private/AutomationReport.h
@@ -72,6 +72,7 @@ public:
 	virtual const bool ExpandInUI() const override;
 	virtual void StopRunningTest() override;
 	virtual bool IsToBeSkipped(FName* OutReason = nullptr, bool* OutWarn = nullptr) const override;
+	virtual bool IsToBeSkippedOnConditions() const override;
 	virtual bool IsToBeSkippedByPropagation() const override;
 	virtual void SetSkipFlag(bool bEnableSkip, const FAutomationTestExcludelistEntry* Template = nullptr, bool bFromPropagation = false) override;
 	virtual TSharedPtr<FAutomationTestExcludeOptions> GetExcludeOptions() override;
diff --git a/Engine/Source/Developer/AutomationController/Public/IAutomationReport.h b/Engine/Source/Developer/AutomationController/Public/IAutomationReport.h
--- a/Engine/Source/Developer/AutomationController/Public/IAutomationReport.h
+++ b/Engine/Source/Developer/AutomationController/Public/IAutomationReport.h
@@ -486,10 +486,19 @@ public:
 	/**
 	* Is the test need to be skipped.
 	*
+	* @param OutReasaon the related reason of the exclusion information
+	* @param OutWarn the related warning of the exclusion
 	* @return true if the test is inside exclude list.
 	*/
 	virtual bool IsToBeSkipped(FName* OutReason = nullptr, bool* OutWarn = nullptr) const = 0;
 
+	/**
+	* Is the test need to be skipped on specific conditions.
+	*
+	* @return true if the test is inside exclude list.
+	*/
+	virtual bool IsToBeSkippedOnConditions() const = 0;
+
 	/**
 	* Is the test is to be skipped through propagation.
 	*
diff --git a/Engine/Source/Developer/AutomationWindow/Private/SAutomationTestItem.cpp b/Engine/Source/Developer/AutomationWindow/Private/SAutomationTestItem.cpp
--- a/Engine/Source/Developer/AutomationWindow/Private/SAutomationTestItem.cpp
+++ b/Engine/Source/Developer/AutomationWindow/Private/SAutomationTestItem.cpp
@@ -38,6 +38,7 @@ void SAutomationTestItem::Construct( const FArguments& InArgs, const TSharedRef<
 	ColumnWidth = InArgs._ColumnWidth;
 	HighlightText = InArgs._HighlightText;
 	OnCheckedStateChangedDelegate = InArgs._OnCheckedStateChanged;
+	IsLocalSession = InArgs._IsLocalSession;
 
 	SMultiColumnTableRow< TSharedPtr< FString > >::Construct( SMultiColumnTableRow< TSharedPtr< FString > >::FArguments(), InOwnerTableView );
 }
@@ -373,12 +374,12 @@ ECheckBoxState SAutomationTestItem::IsToBeSkipped() const
 
 bool SAutomationTestItem::IsDirectlyExcluded() const
 {
-	return WITH_EDITOR && !TestStatus->IsToBeSkippedByPropagation();
+	return WITH_EDITOR && IsLocalSession && !TestStatus->IsToBeSkippedByPropagation();
 }
 
 EVisibility SAutomationTestItem::IsDirectlyExcluded_GetVisibility() const
 {
-	return  TestStatus->IsToBeSkipped() && IsDirectlyExcluded() ? EVisibility::Visible : EVisibility::Collapsed;
+	return  (WITH_EDITOR && IsLocalSession && (IsDirectlyExcluded() || TestStatus->IsToBeSkippedOnConditions())) ? EVisibility::Visible : EVisibility::Collapsed;
 }
 
 FText SAutomationTestItem::GetExcludeReason() const
diff --git a/Engine/Source/Developer/AutomationWindow/Private/SAutomationTestItem.h b/Engine/Source/Developer/AutomationWindow/Private/SAutomationTestItem.h
--- a/Engine/Source/Developer/AutomationWindow/Private/SAutomationTestItem.h
+++ b/Engine/Source/Developer/AutomationWindow/Private/SAutomationTestItem.h
@@ -25,6 +25,7 @@ public:
 	SLATE_BEGIN_ARGS(SAutomationTestItem)
 		{}
 		SLATE_ARGUMENT(float, ColumnWidth)
+		SLATE_ARGUMENT(bool, IsLocalSession)
 		SLATE_ARGUMENT(TSharedPtr<IAutomationReport>, TestStatus)
 		SLATE_ATTRIBUTE(FText, HighlightText)
 		/** Delegate called when a report has it's checkbox clicked */
@@ -202,6 +203,9 @@ private:
 	/** The column width. */
 	float ColumnWidth;
 
+	/** Is Session local. */
+	bool IsLocalSession;
+
 	/** Holds the highlight string for the automation test. */
 	TAttribute<FText> HighlightText;
 
diff --git a/Engine/Source/Developer/AutomationWindow/Private/SAutomationWindow.cpp b/Engine/Source/Developer/AutomationWindow/Private/SAutomationWindow.cpp
--- a/Engine/Source/Developer/AutomationWindow/Private/SAutomationWindow.cpp
+++ b/Engine/Source/Developer/AutomationWindow/Private/SAutomationWindow.cpp
@@ -1722,9 +1722,10 @@ void SAutomationWindow::ClearAutomationUI ()
 TSharedRef<ITableRow> SAutomationWindow::OnGenerateWidgetForTest( TSharedPtr<IAutomationReport> InItem, const TSharedRef<STableViewBase>& OwnerTable )
 {
 	bIsRequestingTests = false;
-	return SNew( SAutomationTestItem, OwnerTable )
-		.TestStatus( InItem )
-		.ColumnWidth( ColumnWidth )
+	return SNew(SAutomationTestItem, OwnerTable)
+		.TestStatus(InItem)
+		.ColumnWidth(ColumnWidth)
+		.IsLocalSession(ActiveSession->IsStandalone())
 		.HighlightText(this, &SAutomationWindow::HandleAutomationHighlightText)
 		.OnCheckedStateChanged(this, &SAutomationWindow::HandleItemCheckBoxCheckedStateChanged);
 }
diff --git a/Engine/Source/Runtime/AutomationTest/Private/AutomationTestExcludelist.cpp b/Engine/Source/Runtime/AutomationTest/Private/AutomationTestExcludelist.cpp
--- a/Engine/Source/Runtime/AutomationTest/Private/AutomationTestExcludelist.cpp
+++ b/Engine/Source/Runtime/AutomationTest/Private/AutomationTestExcludelist.cpp
@@ -136,13 +136,24 @@ FAutomationTestExcludelistEntry* UAutomationTestExcludelist::GetExcludeTestEntry
 
 	const FString NameToCompare = TestName.TrimStartAndEnd().ToLower();
 
+	FAutomationTestExcludelistEntry* OutEntry = nullptr;
+
 	for (auto& Entry : ExcludeTest)
 	{
 		if(NameToCompare.StartsWith(Entry.FullTestName))
 		{
 			if (NameToCompare.Len() == Entry.FullTestName.Len() || NameToCompare.Mid(Entry.FullTestName.Len(), 1) == TEXT("."))
 			{
-				if (RHI.IsEmpty() || Entry.RHIs.Num() == 0 || Entry.RHIs.Contains(*RHI))
+				if (Entry.RHIs.Num() == 0)
+				{
+					return &Entry;
+				}
+				if (RHI.IsEmpty())
+				{
+					OutEntry = &Entry;
+					continue;
+				}
+				if (Entry.RHIs.Contains(*RHI))
 				{
 					return &Entry;
 				}
@@ -150,7 +161,7 @@ FAutomationTestExcludelistEntry* UAutomationTestExcludelist::GetExcludeTestEntry
 		}
 	}
 
-	return nullptr;
+	return OutEntry;
 }
 
 #if WITH_EDITOR
diff --git a/Engine/Source/Runtime/AutomationTest/Public/AutomationTestExcludelist.h b/Engine/Source/Runtime/AutomationTest/Public/AutomationTestExcludelist.h
--- a/Engine/Source/Runtime/AutomationTest/Public/AutomationTestExcludelist.h
+++ b/Engine/Source/Runtime/AutomationTest/Public/AutomationTestExcludelist.h
@@ -114,6 +114,32 @@ struct FAutomationTestExcludelistEntry
 		bIsPropagated = false;
 	}
 
+	/* Has conditional exclusion */
+	bool HasConditions() const
+	{
+		return !RHIs.IsEmpty();
+	}
+
+	/* Remove exclusion conditions, return true if a condition was removed */
+	bool RemoveConditions(const FAutomationTestExcludelistEntry& Entry)
+	{
+		if (!Entry.HasConditions())
+		{
+			return false;
+		}
+
+		bool GotRemoved = false;
+		// Check RHIs
+		int Length = RHIs.Num();
+		if (Length > 0)
+		{
+			RHIs = RHIs.Difference(Entry.RHIs);
+			GotRemoved = GotRemoved || Length != RHIs.Num();
+		}
+
+		return GotRemoved;
+	}
+
 	/* Hold full test name/path */
 	FString FullTestName;
 	/* Is the entry comes from propagation */
