Test Automation - Retry sending a request to run a test if schedule was not confirmed.

#rnx

diff --git a/Engine/Source/Developer/AutomationController/Private/AutomationControllerManager.cpp b/Engine/Source/Developer/AutomationController/Private/AutomationControllerManager.cpp
--- a/Engine/Source/Developer/AutomationController/Private/AutomationControllerManager.cpp
+++ b/Engine/Source/Developer/AutomationController/Private/AutomationControllerManager.cpp
@@ -1084,6 +1084,28 @@ void FAutomationControllerManager::AddPingResult(const FMessageAddress& Responde
 		if ( TestRunningArray[Index].OwnerMessageAddress == ResponderAddress )
 		{
 			TestRunningArray[Index].LastPingTime = 0;
+			// Verify if test was confirmed to be running on the device
+			if (!TestRunningArray[Index].bConfirmed)
+			{
+				// Resend the run test request
+				// Find the game session instance info
+				int32 ClusterIndex;
+				int32 DeviceIndex;
+				verify(DeviceClusterManager.FindDevice(ResponderAddress, ClusterIndex, DeviceIndex));
+				// Verify this device thought it was busy
+				TSharedPtr <IAutomationReport> Report = DeviceClusterManager.GetTest(ClusterIndex, DeviceIndex);
+				if (Report.IsValid())
+				{
+					UE_LOG(LogAutomationController, Verbose, TEXT("Request to run Test '%s' not confirmed from %s. Resending request."), *Report->GetDisplayName(), *ResponderAddress.ToString());
+					MessageEndpoint->Send(
+						FMessageEndpoint::MakeMessage<FAutomationWorkerRunTests>(
+							ExecutionCount, DeviceIndex, Report->GetCommand(), Report->GetDisplayName(), Report->GetFullTestPath(), bSendAnalytics
+						),
+						ResponderAddress
+					);
+				}
+			}
+
 			break;
 		}
 	}
@@ -1493,43 +1515,60 @@ void FAutomationControllerManager::ReportAutomationResult(const TSharedPtr<IAuto
 
 void FAutomationControllerManager::HandleRunTestsReplyMessage(const FAutomationWorkerRunTestsReply& Message, const TSharedRef<IMessageContext, ESPMode::ThreadSafe>& Context)
 {
-	// If we should commit these results
-	if ( Message.ExecutionCount == ExecutionCount )
+	if (Message.ExecutionCount == ExecutionCount)
 	{
-		FAutomationTestResults TestResults;
-
-		TestResults.State = Message.State;
-		TestResults.Duration = Message.Duration;
-
-		// Mark device as back on the market
 		int32 ClusterIndex;
 		int32 DeviceIndex;
-
 		verify(DeviceClusterManager.FindDevice(Context->GetSender(), ClusterIndex, DeviceIndex));
-
-		TestResults.GameInstance = DeviceClusterManager.GetClusterDeviceName(ClusterIndex, DeviceIndex);
-		TestResults.SetEvents(Message.Entries, Message.WarningTotal, Message.ErrorTotal);
-
-		// Verify this device thought it was busy
 		TSharedPtr<IAutomationReport> Report = DeviceClusterManager.GetTest(ClusterIndex, DeviceIndex);
-		if (Report.IsValid())
+
+		// Is it an in-process test reply?
+		if (Message.State == EAutomationState::InProcess)
 		{
-			Report->SetResults(ClusterIndex, CurrentTestPass, TestResults);
+			// Confirm the test is in-process
+			if (Report.IsValid() && Report->GetState(ClusterIndex, CurrentTestPass) == EAutomationState::InProcess && Message.BeautifiedTestName == Report->GetDisplayName())
+			{
+				for (int32 Index = 0; Index < TestRunningArray.Num(); Index++)
+				{
+					if (TestRunningArray[Index].OwnerMessageAddress == Context->GetSender())
+					{
+						TestRunningArray[Index].bConfirmed = true;
+						break;
+					}
+				}
+			}
+		}
+		else
+		{
+			// If we should commit these results
+			FAutomationTestResults TestResults;
+
+			TestResults.State = Message.State;
+			TestResults.Duration = Message.Duration;
+			TestResults.GameInstance = DeviceClusterManager.GetClusterDeviceName(ClusterIndex, DeviceIndex);
+			TestResults.SetEvents(Message.Entries, Message.WarningTotal, Message.ErrorTotal);
+
+			// Mark device as back on the market
+			// Verify this device thought it was busy
+			if (Report.IsValid())
+			{
+				Report->SetResults(ClusterIndex, CurrentTestPass, TestResults);
 
-			const FAutomationTestResults& FinalResults = Report->GetResults(ClusterIndex, CurrentTestPass);
+				const FAutomationTestResults& FinalResults = Report->GetResults(ClusterIndex, CurrentTestPass);
 
-			ReportAutomationResult(Report, ClusterIndex, CurrentTestPass);
+				ReportAutomationResult(Report, ClusterIndex, CurrentTestPass);
 
-			// Gather all of the data relevant to this test for our json reporting.
-			CollectTestResults(Report, FinalResults);
-		}
+				// Gather all of the data relevant to this test for our json reporting.
+				CollectTestResults(Report, FinalResults);
+			}
 
-		// Device is now good to go
-		DeviceClusterManager.SetTest(ClusterIndex, DeviceIndex, NULL);
-	}
+			// Device is now good to go
+			DeviceClusterManager.SetTest(ClusterIndex, DeviceIndex, NULL);
 
-	// Remove the running test
-	RemoveTestRunning(Context->GetSender());
+			// Remove the running test
+			RemoveTestRunning(Context->GetSender());
+		}
+	}
 }
 
 void FAutomationControllerManager::HandleWorkerOfflineMessage( const FAutomationWorkerWorkerOffline& Message, const TSharedRef<IMessageContext, ESPMode::ThreadSafe>& Context )
diff --git a/Engine/Source/Developer/AutomationController/Private/AutomationControllerManager.h b/Engine/Source/Developer/AutomationController/Private/AutomationControllerManager.h
--- a/Engine/Source/Developer/AutomationController/Private/AutomationControllerManager.h
+++ b/Engine/Source/Developer/AutomationController/Private/AutomationControllerManager.h
@@ -567,13 +567,16 @@ private:
 	{
 		FTestRunningInfo( FMessageAddress InMessageAddress ):
 			OwnerMessageAddress( InMessageAddress),
-			LastPingTime( 0.f )
+			LastPingTime( 0.f ),
+			bConfirmed(false)
 		{
 		}
 		/** The test runners message address */
 		FMessageAddress OwnerMessageAddress;
 		/** The time since we had a ping from the instance*/
-		double LastPingTime;
+		float LastPingTime;
+		/** Is Confirmed to be running */
+		bool bConfirmed;
 	};
 
 	/** A array of running tests. */
diff --git a/Engine/Source/Runtime/AutomationMessages/Public/AutomationWorkerMessages.h b/Engine/Source/Runtime/AutomationMessages/Public/AutomationWorkerMessages.h
--- a/Engine/Source/Runtime/AutomationMessages/Public/AutomationWorkerMessages.h
+++ b/Engine/Source/Runtime/AutomationMessages/Public/AutomationWorkerMessages.h
@@ -327,7 +327,7 @@ public:
 
 	/** */
 	UPROPERTY(EditAnywhere, Category="Message")
-	FString TestName;
+	FString BeautifiedTestName;
 
 	/** */
 	UPROPERTY(EditAnywhere, Category="Message")
diff --git a/Engine/Source/Runtime/AutomationWorker/Private/AutomationAnalytics.cpp b/Engine/Source/Runtime/AutomationWorker/Private/AutomationAnalytics.cpp
--- a/Engine/Source/Runtime/AutomationWorker/Private/AutomationAnalytics.cpp
+++ b/Engine/Source/Runtime/AutomationWorker/Private/AutomationAnalytics.cpp
@@ -204,7 +204,7 @@ void FAutomationAnalytics::FireEvent_FPSCapture(const FAutomationPerformanceSnap
 }
 
 
-void FAutomationAnalytics::FireEvent_AutomationTestResults(const FAutomationWorkerRunTestsReply* TestResults, const FString& BeautifiedTestName)
+void FAutomationAnalytics::FireEvent_AutomationTestResults(const FAutomationWorkerRunTestsReply* TestResults, const FString& TestName)
 {
 	// @EventName AutomationTestResults
 	//
@@ -230,8 +230,8 @@ void FAutomationAnalytics::FireEvent_AutomationTestResults(const FAutomationWork
 	{
 		TArray<FAnalyticsEventAttribute> ParamArray;
 		SetInitialParameters(ParamArray);
-		ParamArray.Add(FAnalyticsEventAttribute(GetAutomationParamName(EAutomationAnalyticParam::TestName), TestResults->TestName));
-		ParamArray.Add(FAnalyticsEventAttribute(GetAutomationParamName(EAutomationAnalyticParam::BeautifiedName), BeautifiedTestName));
+		ParamArray.Add(FAnalyticsEventAttribute(GetAutomationParamName(EAutomationAnalyticParam::TestName), TestName));
+		ParamArray.Add(FAnalyticsEventAttribute(GetAutomationParamName(EAutomationAnalyticParam::BeautifiedName), TestResults->BeautifiedTestName));
 		ParamArray.Add(FAnalyticsEventAttribute(GetAutomationParamName(EAutomationAnalyticParam::ExecutionCount), TestResults->ExecutionCount));
 		ParamArray.Add(FAnalyticsEventAttribute(GetAutomationParamName(EAutomationAnalyticParam::IsSuccess), TestResults->State == EAutomationState::Success));
 		ParamArray.Add(FAnalyticsEventAttribute(GetAutomationParamName(EAutomationAnalyticParam::Duration), TestResults->Duration));
diff --git a/Engine/Source/Runtime/AutomationWorker/Private/AutomationAnalytics.h b/Engine/Source/Runtime/AutomationWorker/Private/AutomationAnalytics.h
--- a/Engine/Source/Runtime/AutomationWorker/Private/AutomationAnalytics.h
+++ b/Engine/Source/Runtime/AutomationWorker/Private/AutomationAnalytics.h
@@ -60,7 +60,7 @@ public:
 	static void FireEvent_FPSCapture(const FAutomationPerformanceSnapshot& PerfSnapshot);
 
 	//AutomationTest Event
-	static void FireEvent_AutomationTestResults(const FAutomationWorkerRunTestsReply* TestResults, const FString& BeautifiedTestName);
+	static void FireEvent_AutomationTestResults(const FAutomationWorkerRunTestsReply* TestResults, const FString& TestName);
 
 private:
 	
diff --git a/Engine/Source/Runtime/AutomationWorker/Private/AutomationWorkerModule.cpp b/Engine/Source/Runtime/AutomationWorker/Private/AutomationWorkerModule.cpp
--- a/Engine/Source/Runtime/AutomationWorker/Private/AutomationWorkerModule.cpp
+++ b/Engine/Source/Runtime/AutomationWorker/Private/AutomationWorkerModule.cpp
@@ -200,7 +200,7 @@ void FAutomationWorkerModule::ReportTestComplete()
 			// send the results to the controller
 			FAutomationWorkerRunTestsReply* Message = FMessageEndpoint::MakeMessage<FAutomationWorkerRunTestsReply>();
 
-			Message->TestName = TestName;
+			Message->BeautifiedTestName = BeautifiedTestName;
 			Message->ExecutionCount = ExecutionCount;
 			Message->State = bSuccess ? EAutomationState::Success : EAutomationState::Fail;
 			Message->Duration = ExecutionInfo.Duration;
@@ -215,7 +215,7 @@ void FAutomationWorkerModule::ReportTestComplete()
 				{
 					FAutomationAnalytics::Initialize();
 				}
-				FAutomationAnalytics::FireEvent_AutomationTestResults(Message, BeautifiedTestName);
+				FAutomationAnalytics::FireEvent_AutomationTestResults(Message, TestName);
 				SendAnalyticsEvents(ExecutionInfo.AnalyticsItems);
 			}
 
@@ -489,24 +489,24 @@ void FAutomationWorkerModule::HandleScreenShotAndTraceCapturedWithName(const TAr
 
 void FAutomationWorkerModule::HandleRunTestsMessage( const FAutomationWorkerRunTests& Message, const TSharedRef<IMessageContext, ESPMode::ThreadSafe>& Context )
 {
-	UE_LOG(LogAutomationWorker, Log, TEXT("Received RunTests %s from %s"), *Message.BeautifiedTestName, *Context->GetSender().ToString());
-	LLM_SCOPE_BYNAME(TEXT("AutomationTest/Worker"));
-
 	if (TestRequesterAddress.IsValid() && !TestName.IsEmpty())
 	{
+		// Worker is already running a test.
 		if (TestRequesterAddress == Context->GetSender())
 		{
-			UE_LOG(LogAutomationWorker, Log, TEXT("Worker is already running test '%s' from %s. Request is ignored."), *BeautifiedTestName, *TestRequesterAddress.ToString());
+			// Ignore requests that come from the same sender
+			// It is an extra request sent in case the message did not pass through
+			UE_LOG(LogAutomationWorker, Verbose, TEXT("Worker is already running test '%s' from %s. Request is ignored."), *BeautifiedTestName, *TestRequesterAddress.ToString());
 			return;
 		}
 
-		FString LogMessage = FString::Format(TEXT("Worker is already running test '%s' from %s. '%s' won't be run."), 
-			{ *BeautifiedTestName, *TestRequesterAddress.ToString(), *Message.BeautifiedTestName });
+		FString LogMessage = FString::Format(TEXT("Worker is already running test '{0}' from {1}. Request for '{2}' won't be run."),
+			{ *BeautifiedTestName, *TestRequesterAddress.ToString(), *Message.BeautifiedTestName});
 		UE_LOG(LogAutomationWorker, Warning, TEXT("%s"), *LogMessage);
 
 		// Let the sender know it won't happen
 		FAutomationWorkerRunTestsReply* OutMessage = FMessageEndpoint::MakeMessage<FAutomationWorkerRunTestsReply>();
-		OutMessage->TestName = Message.TestName;
+		OutMessage->BeautifiedTestName = Message.BeautifiedTestName;
 		OutMessage->ExecutionCount = Message.ExecutionCount;
 		OutMessage->State = EAutomationState::Skipped;
 		OutMessage->Entries.Add(FAutomationExecutionEntry(FAutomationEvent(EAutomationEventType::Error, LogMessage)));
@@ -516,6 +516,8 @@ void FAutomationWorkerModule::HandleRunTestsMessage( const FAutomationWorkerRunT
 		return;
 	}
 
+	UE_LOG(LogAutomationWorker, Log, TEXT("Received RunTests %s from %s"), *Message.BeautifiedTestName, *Context->GetSender().ToString());
+
 	// Do we need to skip the test
 	FName SkipReason;
 	bool bWarn(false);
@@ -534,7 +536,7 @@ void FAutomationWorkerModule::HandleRunTestsMessage( const FAutomationWorkerRunT
 		}
 
 		FAutomationWorkerRunTestsReply* OutMessage = FMessageEndpoint::MakeMessage<FAutomationWorkerRunTestsReply>();
-		OutMessage->TestName = Message.TestName;
+		OutMessage->BeautifiedTestName = Message.BeautifiedTestName;
 		OutMessage->ExecutionCount = Message.ExecutionCount;
 		OutMessage->State = EAutomationState::Skipped;
 		OutMessage->Entries.Add(FAutomationExecutionEntry(FAutomationEvent(EAutomationEventType::Info, FString::Printf(TEXT("Skipping test because of exclude list: %s"), *SkipReason.ToString()))));
@@ -555,6 +557,13 @@ void FAutomationWorkerModule::HandleRunTestsMessage( const FAutomationWorkerRunT
 	// We are not executing network command sub-commands right now
 	bExecutingNetworkCommandResults = false;
 
+	// Let the sender know we are running the test
+	FAutomationWorkerRunTestsReply* OutMessage = FMessageEndpoint::MakeMessage<FAutomationWorkerRunTestsReply>();
+	OutMessage->BeautifiedTestName = Message.BeautifiedTestName;
+	OutMessage->ExecutionCount = Message.ExecutionCount;
+	OutMessage->State = EAutomationState::InProcess;
+	MessageEndpoint->Send(OutMessage, Context->GetSender());
+
 	FAutomationTestFramework::Get().StartTestByName(Message.TestName, Message.RoleIndex);
 }
 
