Revert 21824856 (Retry sending a request to run a test if schedule was not confirmed)

#rnx

diff --git a/Engine/Source/Developer/AutomationController/Private/AutomationControllerManager.cpp b/Engine/Source/Developer/AutomationController/Private/AutomationControllerManager.cpp
--- a/Engine/Source/Developer/AutomationController/Private/AutomationControllerManager.cpp
+++ b/Engine/Source/Developer/AutomationController/Private/AutomationControllerManager.cpp
@@ -1084,28 +1084,6 @@ void FAutomationControllerManager::AddPingResult(const FMessageAddress& Responde
 		if ( TestRunningArray[Index].OwnerMessageAddress == ResponderAddress )
 		{
 			TestRunningArray[Index].LastPingTime = 0;
-			// Verify if test was confirmed to be running on the device
-			if (!TestRunningArray[Index].bConfirmed)
-			{
-				// Resend the run test request
-				// Find the game session instance info
-				int32 ClusterIndex;
-				int32 DeviceIndex;
-				verify(DeviceClusterManager.FindDevice(ResponderAddress, ClusterIndex, DeviceIndex));
-				// Verify this device thought it was busy
-				TSharedPtr <IAutomationReport> Report = DeviceClusterManager.GetTest(ClusterIndex, DeviceIndex);
-				if (Report.IsValid())
-				{
-					UE_LOG(LogAutomationController, Verbose, TEXT("Request to run Test '%s' not confirmed from %s. Resending request."), *Report->GetDisplayName(), *ResponderAddress.ToString());
-					MessageEndpoint->Send(
-						FMessageEndpoint::MakeMessage<FAutomationWorkerRunTests>(
-							ExecutionCount, DeviceIndex, Report->GetCommand(), Report->GetDisplayName(), Report->GetFullTestPath(), bSendAnalytics
-						),
-						ResponderAddress
-					);
-				}
-			}
-
 			break;
 		}
 	}
@@ -1515,60 +1493,43 @@ void FAutomationControllerManager::ReportAutomationResult(const TSharedPtr<IAuto
 
 void FAutomationControllerManager::HandleRunTestsReplyMessage(const FAutomationWorkerRunTestsReply& Message, const TSharedRef<IMessageContext, ESPMode::ThreadSafe>& Context)
 {
-	if (Message.ExecutionCount == ExecutionCount)
+	// If we should commit these results
+	if ( Message.ExecutionCount == ExecutionCount )
 	{
-		int32 ClusterIndex;
-		int32 DeviceIndex;
-		verify(DeviceClusterManager.FindDevice(Context->GetSender(), ClusterIndex, DeviceIndex));
-		TSharedPtr<IAutomationReport> Report = DeviceClusterManager.GetTest(ClusterIndex, DeviceIndex);
+		FAutomationTestResults TestResults;
 
-		// Is it an in-process test reply?
-		if (Message.State == EAutomationState::InProcess)
-		{
-			// Confirm the test is in-process
-			if (Report.IsValid() && Report->GetState(ClusterIndex, CurrentTestPass) == EAutomationState::InProcess && Message.BeautifiedTestName == Report->GetDisplayName())
-			{
-				for (int32 Index = 0; Index < TestRunningArray.Num(); Index++)
-				{
-					if (TestRunningArray[Index].OwnerMessageAddress == Context->GetSender())
-					{
-						TestRunningArray[Index].bConfirmed = true;
-						break;
-					}
-				}
-			}
-		}
-		else
-		{
-			// If we should commit these results
-			FAutomationTestResults TestResults;
+		TestResults.State = Message.State;
+		TestResults.Duration = Message.Duration;
 
-			TestResults.State = Message.State;
-			TestResults.Duration = Message.Duration;
-			TestResults.GameInstance = DeviceClusterManager.GetClusterDeviceName(ClusterIndex, DeviceIndex);
-			TestResults.SetEvents(Message.Entries, Message.WarningTotal, Message.ErrorTotal);
+		// Mark device as back on the market
+		int32 ClusterIndex;
+		int32 DeviceIndex;
 
-			// Mark device as back on the market
-			// Verify this device thought it was busy
-			if (Report.IsValid())
-			{
-				Report->SetResults(ClusterIndex, CurrentTestPass, TestResults);
+		verify(DeviceClusterManager.FindDevice(Context->GetSender(), ClusterIndex, DeviceIndex));
 
-				const FAutomationTestResults& FinalResults = Report->GetResults(ClusterIndex, CurrentTestPass);
+		TestResults.GameInstance = DeviceClusterManager.GetClusterDeviceName(ClusterIndex, DeviceIndex);
+		TestResults.SetEvents(Message.Entries, Message.WarningTotal, Message.ErrorTotal);
 
-				ReportAutomationResult(Report, ClusterIndex, CurrentTestPass);
+		// Verify this device thought it was busy
+		TSharedPtr<IAutomationReport> Report = DeviceClusterManager.GetTest(ClusterIndex, DeviceIndex);
+		if (Report.IsValid())
+		{
+			Report->SetResults(ClusterIndex, CurrentTestPass, TestResults);
 
-				// Gather all of the data relevant to this test for our json reporting.
-				CollectTestResults(Report, FinalResults);
-			}
+			const FAutomationTestResults& FinalResults = Report->GetResults(ClusterIndex, CurrentTestPass);
 
-			// Device is now good to go
-			DeviceClusterManager.SetTest(ClusterIndex, DeviceIndex, NULL);
+			ReportAutomationResult(Report, ClusterIndex, CurrentTestPass);
 
-			// Remove the running test
-			RemoveTestRunning(Context->GetSender());
+			// Gather all of the data relevant to this test for our json reporting.
+			CollectTestResults(Report, FinalResults);
 		}
+
+		// Device is now good to go
+		DeviceClusterManager.SetTest(ClusterIndex, DeviceIndex, NULL);
 	}
+
+	// Remove the running test
+	RemoveTestRunning(Context->GetSender());
 }
 
 void FAutomationControllerManager::HandleWorkerOfflineMessage( const FAutomationWorkerWorkerOffline& Message, const TSharedRef<IMessageContext, ESPMode::ThreadSafe>& Context )
diff --git a/Engine/Source/Developer/AutomationController/Private/AutomationControllerManager.h b/Engine/Source/Developer/AutomationController/Private/AutomationControllerManager.h
--- a/Engine/Source/Developer/AutomationController/Private/AutomationControllerManager.h
+++ b/Engine/Source/Developer/AutomationController/Private/AutomationControllerManager.h
@@ -567,16 +567,13 @@ private:
 	{
 		FTestRunningInfo( FMessageAddress InMessageAddress ):
 			OwnerMessageAddress( InMessageAddress),
-			LastPingTime( 0.f ),
-			bConfirmed(false)
+			LastPingTime( 0.f )
 		{
 		}
 		/** The test runners message address */
 		FMessageAddress OwnerMessageAddress;
 		/** The time since we had a ping from the instance*/
-		double LastPingTime;
-		/** Is Confirmed to be running */
-		bool bConfirmed;
+		float LastPingTime;
 	};
 
 	/** A array of running tests. */
diff --git a/Engine/Source/Runtime/AutomationMessages/Public/AutomationWorkerMessages.h b/Engine/Source/Runtime/AutomationMessages/Public/AutomationWorkerMessages.h
--- a/Engine/Source/Runtime/AutomationMessages/Public/AutomationWorkerMessages.h
+++ b/Engine/Source/Runtime/AutomationMessages/Public/AutomationWorkerMessages.h
@@ -327,7 +327,7 @@ public:
 
 	/** */
 	UPROPERTY(EditAnywhere, Category="Message")
-	FString BeautifiedTestName;
+	FString TestName;
 
 	/** */
 	UPROPERTY(EditAnywhere, Category="Message")
diff --git a/Engine/Source/Runtime/AutomationWorker/Private/AutomationAnalytics.cpp b/Engine/Source/Runtime/AutomationWorker/Private/AutomationAnalytics.cpp
--- a/Engine/Source/Runtime/AutomationWorker/Private/AutomationAnalytics.cpp
+++ b/Engine/Source/Runtime/AutomationWorker/Private/AutomationAnalytics.cpp
@@ -204,7 +204,7 @@ void FAutomationAnalytics::FireEvent_FPSCapture(const FAutomationPerformanceSnap
 }
 
 
-void FAutomationAnalytics::FireEvent_AutomationTestResults(const FAutomationWorkerRunTestsReply* TestResults, const FString& TestName)
+void FAutomationAnalytics::FireEvent_AutomationTestResults(const FAutomationWorkerRunTestsReply* TestResults, const FString& BeautifiedTestName)
 {
 	// @EventName AutomationTestResults
 	//
@@ -230,8 +230,8 @@ void FAutomationAnalytics::FireEvent_AutomationTestResults(const FAutomationWork
 	{
 		TArray<FAnalyticsEventAttribute> ParamArray;
 		SetInitialParameters(ParamArray);
-		ParamArray.Add(FAnalyticsEventAttribute(GetAutomationParamName(EAutomationAnalyticParam::TestName), TestName));
-		ParamArray.Add(FAnalyticsEventAttribute(GetAutomationParamName(EAutomationAnalyticParam::BeautifiedName), TestResults->BeautifiedTestName));
+		ParamArray.Add(FAnalyticsEventAttribute(GetAutomationParamName(EAutomationAnalyticParam::TestName), TestResults->TestName));
+		ParamArray.Add(FAnalyticsEventAttribute(GetAutomationParamName(EAutomationAnalyticParam::BeautifiedName), BeautifiedTestName));
 		ParamArray.Add(FAnalyticsEventAttribute(GetAutomationParamName(EAutomationAnalyticParam::ExecutionCount), TestResults->ExecutionCount));
 		ParamArray.Add(FAnalyticsEventAttribute(GetAutomationParamName(EAutomationAnalyticParam::IsSuccess), TestResults->State == EAutomationState::Success));
 		ParamArray.Add(FAnalyticsEventAttribute(GetAutomationParamName(EAutomationAnalyticParam::Duration), TestResults->Duration));
diff --git a/Engine/Source/Runtime/AutomationWorker/Private/AutomationAnalytics.h b/Engine/Source/Runtime/AutomationWorker/Private/AutomationAnalytics.h
--- a/Engine/Source/Runtime/AutomationWorker/Private/AutomationAnalytics.h
+++ b/Engine/Source/Runtime/AutomationWorker/Private/AutomationAnalytics.h
@@ -60,7 +60,7 @@ public:
 	static void FireEvent_FPSCapture(const FAutomationPerformanceSnapshot& PerfSnapshot);
 
 	//AutomationTest Event
-	static void FireEvent_AutomationTestResults(const FAutomationWorkerRunTestsReply* TestResults, const FString& TestName);
+	static void FireEvent_AutomationTestResults(const FAutomationWorkerRunTestsReply* TestResults, const FString& BeautifiedTestName);
 
 private:
 	
diff --git a/Engine/Source/Runtime/AutomationWorker/Private/AutomationWorkerModule.cpp b/Engine/Source/Runtime/AutomationWorker/Private/AutomationWorkerModule.cpp
--- a/Engine/Source/Runtime/AutomationWorker/Private/AutomationWorkerModule.cpp
+++ b/Engine/Source/Runtime/AutomationWorker/Private/AutomationWorkerModule.cpp
@@ -200,7 +200,7 @@ void FAutomationWorkerModule::ReportTestComplete()
 			// send the results to the controller
 			FAutomationWorkerRunTestsReply* Message = FMessageEndpoint::MakeMessage<FAutomationWorkerRunTestsReply>();
 
-			Message->BeautifiedTestName = BeautifiedTestName;
+			Message->TestName = TestName;
 			Message->ExecutionCount = ExecutionCount;
 			Message->State = bSuccess ? EAutomationState::Success : EAutomationState::Fail;
 			Message->Duration = ExecutionInfo.Duration;
@@ -215,7 +215,7 @@ void FAutomationWorkerModule::ReportTestComplete()
 				{
 					FAutomationAnalytics::Initialize();
 				}
-				FAutomationAnalytics::FireEvent_AutomationTestResults(Message, TestName);
+				FAutomationAnalytics::FireEvent_AutomationTestResults(Message, BeautifiedTestName);
 				SendAnalyticsEvents(ExecutionInfo.AnalyticsItems);
 			}
 
@@ -489,26 +489,25 @@ void FAutomationWorkerModule::HandleScreenShotAndTraceCapturedWithName(const TAr
 
 void FAutomationWorkerModule::HandleRunTestsMessage( const FAutomationWorkerRunTests& Message, const TSharedRef<IMessageContext, ESPMode::ThreadSafe>& Context )
 {
+	UE_LOG(LogAutomationWorker, Log, TEXT("Received RunTests %s from %s"), *Message.BeautifiedTestName, *Context->GetSender().ToString());
+
 	LLM_SCOPE_BYNAME(TEXT("AutomationTest/Worker"));
 	
 	if (TestRequesterAddress.IsValid() && !TestName.IsEmpty())
 	{
-		// Worker is already running a test.
 		if (TestRequesterAddress == Context->GetSender())
 		{
-			// Ignore requests that come from the same sender
-			// It is an extra request sent in case the message did not pass through
-			UE_LOG(LogAutomationWorker, Verbose, TEXT("Worker is already running test '%s' from %s. Request is ignored."), *BeautifiedTestName, *TestRequesterAddress.ToString());
+			UE_LOG(LogAutomationWorker, Log, TEXT("Worker is already running test '%s' from %s. Request is ignored."), *BeautifiedTestName, *TestRequesterAddress.ToString());
 			return;
 		}
 
-		FString LogMessage = FString::Format(TEXT("Worker is already running test '{0}' from {1}. Request for '{2}' won't be run."),
-			{ *BeautifiedTestName, *TestRequesterAddress.ToString(), *Message.BeautifiedTestName});
+		FString LogMessage = FString::Format(TEXT("Worker is already running test '%s' from %s. '%s' won't be run."), 
+			{ *BeautifiedTestName, *TestRequesterAddress.ToString(), *Message.BeautifiedTestName });
 		UE_LOG(LogAutomationWorker, Warning, TEXT("%s"), *LogMessage);
 
 		// Let the sender know it won't happen
 		FAutomationWorkerRunTestsReply* OutMessage = FMessageEndpoint::MakeMessage<FAutomationWorkerRunTestsReply>();
-		OutMessage->BeautifiedTestName = Message.BeautifiedTestName;
+		OutMessage->TestName = Message.TestName;
 		OutMessage->ExecutionCount = Message.ExecutionCount;
 		OutMessage->State = EAutomationState::Skipped;
 		OutMessage->Entries.Add(FAutomationExecutionEntry(FAutomationEvent(EAutomationEventType::Error, LogMessage)));
@@ -518,8 +517,6 @@ void FAutomationWorkerModule::HandleRunTestsMessage( const FAutomationWorkerRunT
 		return;
 	}
 
-	UE_LOG(LogAutomationWorker, Log, TEXT("Received RunTests %s from %s"), *Message.BeautifiedTestName, *Context->GetSender().ToString());
-
 	// Do we need to skip the test
 	FName SkipReason;
 	bool bWarn(false);
@@ -538,7 +535,7 @@ void FAutomationWorkerModule::HandleRunTestsMessage( const FAutomationWorkerRunT
 		}
 
 		FAutomationWorkerRunTestsReply* OutMessage = FMessageEndpoint::MakeMessage<FAutomationWorkerRunTestsReply>();
-		OutMessage->BeautifiedTestName = Message.BeautifiedTestName;
+		OutMessage->TestName = Message.TestName;
 		OutMessage->ExecutionCount = Message.ExecutionCount;
 		OutMessage->State = EAutomationState::Skipped;
 		OutMessage->Entries.Add(FAutomationExecutionEntry(FAutomationEvent(EAutomationEventType::Info, FString::Printf(TEXT("Skipping test because of exclude list: %s"), *SkipReason.ToString()))));
@@ -559,13 +556,6 @@ void FAutomationWorkerModule::HandleRunTestsMessage( const FAutomationWorkerRunT
 	// We are not executing network command sub-commands right now
 	bExecutingNetworkCommandResults = false;
 
-	// Let the sender know we are running the test
-	FAutomationWorkerRunTestsReply* OutMessage = FMessageEndpoint::MakeMessage<FAutomationWorkerRunTestsReply>();
-	OutMessage->BeautifiedTestName = Message.BeautifiedTestName;
-	OutMessage->ExecutionCount = Message.ExecutionCount;
-	OutMessage->State = EAutomationState::InProcess;
-	MessageEndpoint->Send(OutMessage, Context->GetSender());
-
 	FAutomationTestFramework::Get().StartTestByName(Message.TestName, Message.RoleIndex);
 }
 
