Follow up to 21896012 - change LastPingTime type from float to double

#rnx

diff --git a/Engine/Source/Developer/AutomationController/Private/AutomationControllerManager.h b/Engine/Source/Developer/AutomationController/Private/AutomationControllerManager.h
--- a/Engine/Source/Developer/AutomationController/Private/AutomationControllerManager.h
+++ b/Engine/Source/Developer/AutomationController/Private/AutomationControllerManager.h
@@ -573,7 +573,7 @@ private:
 		/** The test runners message address */
 		FMessageAddress OwnerMessageAddress;
 		/** The time since we had a ping from the instance*/
-		float LastPingTime;
+		double LastPingTime;
 	};
 
 	/** A array of running tests. */
