PR #7401: Fix/AutomationCommandLine (Contributed by [REDACTED_PERSON])

diff --git a/Engine/Source/Developer/AutomationController/Private/AutomationCommandline.cpp b/Engine/Source/Developer/AutomationController/Private/AutomationCommandline.cpp
--- a/Engine/Source/Developer/AutomationController/Private/AutomationCommandline.cpp
+++ b/Engine/Source/Developer/AutomationController/Private/AutomationCommandline.cpp
@@ -527,7 +527,7 @@ public:
 					}
 					else if (StringCommand.Contains(TEXT("Medium")))
 					{
-						AutomationController->SetRequestedTestFlags(EAutomationTestFlags::MediumPriority);
+						AutomationController->SetRequestedTestFlags(EAutomationTestFlags::MediumPriorityAndAbove);
 					}
 					else if (StringCommand.Contains(TEXT("High")))
 					{
@@ -535,7 +535,7 @@ public:
 					}
 					else if (StringCommand.Contains(TEXT("Critical")))
 					{
-						AutomationController->SetRequestedTestFlags(EAutomationTestFlags::ClientContext);
+						AutomationController->SetRequestedTestFlags(EAutomationTestFlags::CriticalPriority);
 					}
 					else if (StringCommand.Contains(TEXT("None")))
 					{
@@ -622,9 +622,11 @@ public:
 					Ar.Logf(TEXT("\tAutomation StartRemoteSession <sessionid>"));
 					Ar.Logf(TEXT("\tAutomation List"));
 					Ar.Logf(TEXT("\tAutomation RunTests <test string>"));
-					Ar.Logf(TEXT("\tAutomation RunAll "));
+					Ar.Logf(TEXT("\tAutomation RunAll"));
 					Ar.Logf(TEXT("\tAutomation RunFilter <filter name>"));
 					Ar.Logf(TEXT("\tAutomation SetFilter <filter name>"));
+					Ar.Logf(TEXT("\tAutomation SetMinimumPriority <minimum priority>"));
+					Ar.Logf(TEXT("\tAutomation SetPriority <priority>"));
 					Ar.Logf(TEXT("\tAutomation Quit"));
 					bHandled = false;
 				}
