AutomationController - Fixing end queue message not being triggered
Prevent multiple Run Test to be queued at the same time (not supported)
Rename Filter: to StartsWith:
Some more code cleanup

#rnx

diff --git a/Engine/Source/Developer/AutomationController/Private/AutomationCommandline.cpp b/Engine/Source/Developer/AutomationController/Private/AutomationCommandline.cpp
--- a/Engine/Source/Developer/AutomationController/Private/AutomationCommandline.cpp
+++ b/Engine/Source/Developer/AutomationController/Private/AutomationCommandline.cpp
@@ -22,8 +22,8 @@ DEFINE_LOG_CATEGORY_STATIC(LogAutomationCommandLine, Log, All);
 /** States for running the automation process */
 enum class EAutomationTestState : uint8
 {
-	Initializing,		// 
 	Idle,				// Automation process is not running
+	Initializing,		//
 	FindWorkers,		// Find workers to run the tests
 	RequestTests,		// Find the tests that can be run on the workers
 	DoingRequestedWork,	// Do whatever was requested from the commandline
@@ -35,7 +35,7 @@ enum class EAutomationCommand : uint8
 	ListAllTests,			//List all tests for the session
 	RunCommandLineTests,	//Run only tests that are listed on the commandline
 	RunAll,					//Run all the tests that are supported
-	RunFilter,              //
+	RunFilter,              //Run only tests that are tagged with this filter
 	Quit					//quit the app when tests are done
 };
 
@@ -51,6 +51,7 @@ public:
 		DelayTimer = DefaultDelayTimer;
 		FindWorkersTimeout = DefaultFindWorkersTimeout;
 		FindWorkerAttempts = 0;
+		TestCount = 0;
 	}
 
 	void Init()
@@ -76,12 +77,13 @@ public:
 			TestsRefreshedHandle = AutomationController->OnTestsRefreshed().AddRaw(this, &FAutomationExecCmd::HandleRefreshTestCallback);
 		}
 
-		TickHandler = FTSTicker::GetCoreTicker().AddTicker(FTickerDelegate::CreateRaw(this, &FAutomationExecCmd::Tick));
+		if (!TickHandler.IsValid()) {
+			TickHandler = FTSTicker::GetCoreTicker().AddTicker(FTickerDelegate::CreateRaw(this, &FAutomationExecCmd::Tick));
+		}
 
 		int32 NumTestLoops = 1;
 		FParse::Value(FCommandLine::Get(), TEXT("TestLoops="), NumTestLoops);
 		AutomationController->SetNumPasses(NumTestLoops);
-		TestCount = 0;
 		SetUpFilterMapping();
 	}
 
@@ -113,12 +115,8 @@ public:
 		// If the automation controller is no longer processing and we've reached the final stage of testing
 		if ((AutomationController->GetTestState() != EAutomationControllerModuleState::Running) && (AutomationTestState == EAutomationTestState::Complete) && (AutomationCommandQueue.Num() == 0))
 		{
-			// If an actual test was ran we then will let the user know how many of them were ran.
-			if (TestCount > 0)
-			{
-				UE_LOG(LogAutomationCommandLine, Display, TEXT("...Automation Test Queue Empty %d tests performed."), TestCount);
-				TestCount = 0;
-			}
+			UE_LOG(LogAutomationCommandLine, Display, TEXT("...Automation Test Queue Empty %d tests performed."), TestCount);
+			TestCount = 0;
 			return true;
 		}
 		return false;
@@ -136,7 +134,7 @@ public:
 		UAutomationControllerSettings* Settings = UAutomationControllerSettings::StaticClass()->GetDefaultObject<UAutomationControllerSettings>();
 
 		// iterate through the arguments to build a filter list by doing the following -
-		// 1) If argument is a filter (filter:system) then make sure we only filter-in tests that start with that filter
+		// 1) If argument is a filter (StartsWith:system) then make sure we only filter-in tests that start with that filter
 		// 2) If argument is a group then expand that group into multiple filters based on ini entries
 		// 3) Otherwise just substring match (default behavior in 4.22 and earlier).
 		FAutomationGroupFilter* FilterAny = new FAutomationGroupFilter();
@@ -144,7 +142,7 @@ public:
 		for (int32 ArgumentIndex = 0; ArgumentIndex < ArgumentNames.Num(); ++ArgumentIndex)
 		{
 			const FString GroupPrefix = TEXT("Group:");
-			const FString FilterPrefix = TEXT("Filter:");
+			const FString FilterPrefix = TEXT("StartsWith:");
 
 			FString ArgumentName = ArgumentNames[ArgumentIndex].TrimStartAndEnd();
 
@@ -262,7 +260,7 @@ public:
 
 	void HandleRefreshTestCallback()
 	{
-		TArray<FString> AllTestNames;
+		TArray<FString> FilteredTestNames;
 
 		// This is called by the controller manager when it receives responses. We want to make sure it has a device, and we
 		// want to make sure it's called while we're waiting for a response
@@ -277,15 +275,15 @@ public:
 		TSharedPtr <AutomationFilterCollection> AutomationFilters = MakeShareable(new AutomationFilterCollection());
 		AutomationController->SetFilter(AutomationFilters);
 		AutomationController->SetVisibleTestsEnabled(true);
-		AutomationController->GetEnabledTestNames(AllTestNames);
+		AutomationController->GetEnabledTestNames(FilteredTestNames);
 
 		//assume we won't run any tests
 		bool bRunTests = false;
 
 		if (AutomationCommand == EAutomationCommand::ListAllTests)
 		{
-			UE_LOG(LogAutomationCommandLine, Display, TEXT("Found %d Automation Tests"), AllTestNames.Num());
-			for ( const FString& TestName : AllTestNames )
+			UE_LOG(LogAutomationCommandLine, Display, TEXT("Found %d Automation Tests"), FilteredTestNames.Num());
+			for ( const FString& TestName : FilteredTestNames)
 			{
 				UE_LOG(LogAutomationCommandLine, Display, TEXT("\t'%s'"), *TestName);
 			}
@@ -295,7 +293,6 @@ public:
 		}
 		else if (AutomationCommand == EAutomationCommand::RunCommandLineTests)
 		{
-			TArray<FString> FilteredTestNames;
 			GenerateTestNamesFromCommandLine(AutomationFilters, FilteredTestNames);
 			
 			if (FilteredTestNames.Num() == 0)
@@ -314,27 +311,26 @@ public:
 
 			if (FilteredTestNames.Num())
 			{
-				AutomationController->StopTests();
-				AutomationController->SetEnabledTests(FilteredTestNames);
 				bRunTests = true;
 			}
 			else
 			{
 				AutomationTestState = EAutomationTestState::Complete;
 			}
-
-			// Clear delegate to avoid re-running tests due to multiple delegates being added or when refreshing session frontend
-			// The delegate will be readded in Init whenever a new command is executed
-			AutomationController->OnTestsRefreshed().Remove(TestsRefreshedHandle);
-			TestsRefreshedHandle.Reset();
 		}
 		else if (AutomationCommand == EAutomationCommand::RunFilter)
 		{
 			if (FilterMaps.Contains(StringCommand))
 			{
-				UE_LOG(LogAutomationCommandLine, Display, TEXT("Running %i Automation Tests"), AllTestNames.Num());
-				AutomationController->SetEnabledTests(AllTestNames);
-				bRunTests = true;
+				UE_LOG(LogAutomationCommandLine, Display, TEXT("Running %i Automation Tests"), FilteredTestNames.Num());
+				if (FilteredTestNames.Num() > 0)
+				{
+					bRunTests = true;
+				}
+				else
+				{
+					AutomationTestState = EAutomationTestState::Complete;
+				}
 			}
 			else
 			{
@@ -351,11 +347,19 @@ public:
 		else if (AutomationCommand == EAutomationCommand::RunAll)
 		{
 			bRunTests = true;
-			TestCount = AllTestNames.Num();
 		}
 
 		if (bRunTests)
 		{
+			AutomationController->StopTests();
+			AutomationController->SetEnabledTests(FilteredTestNames);
+			TestCount = FilteredTestNames.Num();
+
+			// Clear delegate to avoid re-running tests due to multiple delegates being added or when refreshing session frontend
+			// The delegate will be readded in Init whenever a new command is executed
+			AutomationController->OnTestsRefreshed().Remove(TestsRefreshedHandle);
+			TestsRefreshedHandle.Reset();
+
 			AutomationController->RunTests();
 
 			// Set state to monitoring to check for test completion
@@ -440,22 +444,38 @@ public:
 						}
 					}
 					UE_LOG(LogAutomationCommandLine, Log, TEXT("Shutting down. GIsCriticalError=%d"), GIsCriticalError);
+					// some tools parse this.
 					UE_LOG(LogAutomationCommandLine, Display, TEXT("**** TEST COMPLETE. EXIT CODE: %d ****"), GIsCriticalError ? -1 : 0);
 					FPlatformMisc::RequestExitWithStatus(true, GIsCriticalError ? -1 : 0);
+					// We have finished the testing, and results are available
 					AutomationTestState = EAutomationTestState::Complete;
 				}
 				break;
 			}
 		}
 
-		return !IsTestingComplete();
-
-
+		if (IsTestingComplete())
+		{
+			AutomationTestState = EAutomationTestState::Idle;
+			TickHandler.Reset();
+			return false;
+		}
+		return true;
+	}
 
+	bool IsRunTestQueued()
+	{
+		for (auto Command : AutomationCommandQueue)
+		{
+			if (Command == EAutomationCommand::RunCommandLineTests
+				|| Command == EAutomationCommand::RunAll
+				|| Command == EAutomationCommand::RunFilter)
+			{
+				return true;
+			}
+		}
 
-		
-		// some tools parse this.
-		// We have finished the testing, and results are available
+		return false;
 	}
 	
 	/** Console commands, see embeded usage statement **/
@@ -465,27 +485,21 @@ public:
 		// Track whether we have a flag we care about passing through.
 		FString FlagToUse = "";
 
-		// Hackiest hack to ever hack a hack to get this test running.
-		if (FParse::Command(&Cmd, TEXT("RunPerfTests")))
-		{
-			Cmd = TEXT("Automation RunFilter Perf");
-		}
-		else if (FParse::Command(&Cmd, TEXT("RunProductTests")))
-		{
-			Cmd = TEXT("Automation RunFilter Product");
-		}
-
-		//figure out if we are handling this request
+		// figure out if we are handling this request
 		if (FParse::Command(&Cmd, TEXT("Automation")))
 		{
-			StringCommand.Empty();
+			if (AutomationTestState != EAutomationTestState::Idle)
+			{
+				Ar.Logf(TEXT("Automation: Skipping this request because already executing tests!"));
+				return false;
+			}
 
 			TArray<FString> CommandList;
-			StringCommand = Cmd;
-			StringCommand.ParseIntoArray(CommandList, TEXT(";"), true);
+			FString(Cmd).ParseIntoArray(CommandList, TEXT(";"), true);
 
-			//assume we handle this
 			Init();
+
+			//assume we handle this
 			bHandled = true;
 
 			for (int CommandIndex = 0; CommandIndex < CommandList.Num(); ++CommandIndex)
@@ -496,7 +510,7 @@ public:
 					FString SessionString = TempCmd;
 					if (!FGuid::Parse(SessionString, SessionID))
 					{
-						Ar.Logf(TEXT("%s is not a valid session guid!"), *SessionString);
+						Ar.Logf(TEXT("Automation: %s is not a valid session guid!"), *SessionString);
 						bHandled = false;
 						break;
 					}
@@ -507,101 +521,117 @@ public:
 				}
 				else if (FParse::Command(&TempCmd, TEXT("RunTests")) || FParse::Command(&TempCmd, TEXT("RunTest")))
 				{
+					//only one of these should be used
+					if (IsRunTestQueued())
+					{
+						Ar.Logf(TEXT("Automation: A test run is already Queued: %s. Only one run is supported at a time."), *StringCommand);
+						continue;
+					}
+
 					if ( FParse::Command(&TempCmd, TEXT("Now")) )
 					{
 						DelayTimer = 0.0f;
 					}
 
-					//only one of these should be used
 					StringCommand = TempCmd;
 					Ar.Logf(TEXT("Automation: RunTests='%s' Queued."), *StringCommand);
 					AutomationCommandQueue.Add(EAutomationCommand::RunCommandLineTests);
 				}
 				else if (FParse::Command(&TempCmd, TEXT("SetMinimumPriority")))
 				{
-					StringCommand = TempCmd;
-					Ar.Logf(TEXT("Setting minimum priority of cases to run to: %s"), *StringCommand);
-					if (StringCommand.Contains(TEXT("Low")))
+					FlagToUse = TempCmd;
+					Ar.Logf(TEXT("Automation: Setting minimum priority of cases to run to: %s"), *FlagToUse);
+					if (FlagToUse.Contains(TEXT("Low")))
 					{
 						AutomationController->SetRequestedTestFlags(EAutomationTestFlags::PriorityMask);
 					}
-					else if (StringCommand.Contains(TEXT("Medium")))
+					else if (FlagToUse.Contains(TEXT("Medium")))
 					{
 						AutomationController->SetRequestedTestFlags(EAutomationTestFlags::MediumPriorityAndAbove);
 					}
-					else if (StringCommand.Contains(TEXT("High")))
+					else if (FlagToUse.Contains(TEXT("High")))
 					{
 						AutomationController->SetRequestedTestFlags(EAutomationTestFlags::HighPriorityAndAbove);
 					}
-					else if (StringCommand.Contains(TEXT("Critical")))
+					else if (FlagToUse.Contains(TEXT("Critical")))
 					{
 						AutomationController->SetRequestedTestFlags(EAutomationTestFlags::CriticalPriority);
 					}
-					else if (StringCommand.Contains(TEXT("None")))
+					else if (FlagToUse.Contains(TEXT("None")))
 					{
 						AutomationController->SetRequestedTestFlags(0);
 					}
 					else
 					{
-						Ar.Logf(TEXT("%s is not a valid priority!\nValid priorities are Critical, High, Medium, Low, None"), *StringCommand);
+						Ar.Logf(TEXT("Automation: %s is not a valid priority!\nValid priorities are Critical, High, Medium, Low, None"), *FlagToUse);
 					}
 				}
 				else if (FParse::Command(&TempCmd, TEXT("SetPriority")))
 				{
-					StringCommand = TempCmd;
-					Ar.Logf(TEXT("Setting explicit priority of cases to run to: %s"), *StringCommand);
-					if (StringCommand.Contains(TEXT("Low")))
+					FlagToUse = TempCmd;
+					Ar.Logf(TEXT("Setting explicit priority of cases to run to: %s"), *FlagToUse);
+					if (FlagToUse.Contains(TEXT("Low")))
 					{
 						AutomationController->SetRequestedTestFlags(EAutomationTestFlags::LowPriority);
 					}
-					else if (StringCommand.Contains(TEXT("Medium")))
+					else if (FlagToUse.Contains(TEXT("Medium")))
 					{
 						AutomationController->SetRequestedTestFlags(EAutomationTestFlags::MediumPriority);
 					}
-					else if (StringCommand.Contains(TEXT("High")))
+					else if (FlagToUse.Contains(TEXT("High")))
 					{
 						AutomationController->SetRequestedTestFlags(EAutomationTestFlags::HighPriority);
 					}
-					else if (StringCommand.Contains(TEXT("Critical")))
+					else if (FlagToUse.Contains(TEXT("Critical")))
 					{
 						AutomationController->SetRequestedTestFlags(EAutomationTestFlags::CriticalPriority);
 					}
-					else if (StringCommand.Contains(TEXT("None")))
+					else if (FlagToUse.Contains(TEXT("None")))
 					{
 						AutomationController->SetRequestedTestFlags(0);
 					}
 
 					else
 					{
-						Ar.Logf(TEXT("%s is not a valid priority!\nValid priorities are Critical, High, Medium, Low, None"), *StringCommand);
+						Ar.Logf(TEXT("Automation: %s is not a valid priority!\nValid priorities are Critical, High, Medium, Low, None"), *StringCommand);
 					}
 				}
 				else if (FParse::Command(&TempCmd, TEXT("RunFilter")))
 				{
-					FlagToUse = TempCmd;
 					//only one of these should be used
+					if (IsRunTestQueued())
+					{
+						Ar.Logf(TEXT("Automation: A test run is already Queued: %s. Only one run is supported at a time."), *StringCommand);
+						continue;
+					}
+					FlagToUse = TempCmd;
 					StringCommand = TempCmd;
 					if (FilterMaps.Contains(FlagToUse))
 					{
 						AutomationController->SetRequestedTestFlags(FilterMaps[FlagToUse]);
-						Ar.Logf(TEXT("Running all tests for filter: %s"), *FlagToUse);
+						Ar.Logf(TEXT("Automation: RunFilter='%s' Queued."), *FlagToUse);
 					}
 					AutomationCommandQueue.Add(EAutomationCommand::RunFilter);
 				}
 				else if (FParse::Command(&TempCmd, TEXT("SetFilter")))
 				{
 					FlagToUse = TempCmd;
-					StringCommand = TempCmd;
 					if (FilterMaps.Contains(FlagToUse))
 					{
 						AutomationController->SetRequestedTestFlags(FilterMaps[FlagToUse]);
-						Ar.Logf(TEXT("Setting test filter: %s"), *FlagToUse);
+						Ar.Logf(TEXT("Automation: Setting test filter: %s"), *FlagToUse);
 					}
 				}
 				else if (FParse::Command(&TempCmd, TEXT("RunAll")))
 				{
+					//only one of these should be used
+					if (IsRunTestQueued())
+					{
+						Ar.Logf(TEXT("Automation: A test run is already Queued: %s. Only one run is supported at a time."), *StringCommand);
+						continue;
+					}
 					AutomationCommandQueue.Add(EAutomationCommand::RunAll);
-					Ar.Logf(TEXT("Running all available automated tests for this program. NOTE: This may take a while."));
+					Ar.Logf(TEXT("Automation: RunAll Queued. NOTE: This may take a while."));
 				}
 				else if (FParse::Command(&TempCmd, TEXT("Quit")))
 				{
diff --git a/Engine/Source/Developer/AutomationController/Private/AutomationControllerManager.cpp b/Engine/Source/Developer/AutomationController/Private/AutomationControllerManager.cpp
--- a/Engine/Source/Developer/AutomationController/Private/AutomationControllerManager.cpp
+++ b/Engine/Source/Developer/AutomationController/Private/AutomationControllerManager.cpp
@@ -425,6 +425,7 @@ void FAutomationControllerManager::Init()
 	AutomationTestState = EAutomationControllerModuleState::Disabled;
 	bTestResultsAvailable = false;
 	bSendAnalytics = FParse::Param(FCommandLine::Get(), TEXT("SendAutomationAnalytics"));
+	RequestedTestFlags = EAutomationTestFlags::SmokeFilter | EAutomationTestFlags::EngineFilter | EAutomationTestFlags::ProductFilter | EAutomationTestFlags::PerfFilter;
 }
 
 void FAutomationControllerManager::RequestLoadAsset(const FString& InAssetName)
