Improve device timeout error by differenciating DeviceName and GameInstance from the Automation device cluster manager API

diff --git a/Engine/Source/Developer/AutomationController/Private/AutomationControllerManager.cpp b/Engine/Source/Developer/AutomationController/Private/AutomationControllerManager.cpp
--- a/Engine/Source/Developer/AutomationController/Private/AutomationControllerManager.cpp
+++ b/Engine/Source/Developer/AutomationController/Private/AutomationControllerManager.cpp
@@ -771,7 +771,7 @@ void FAutomationControllerManager::ExecuteNextTask( int32 ClusterIndex, OUT bool
 						for (int32 AddressIndex = 0; AddressIndex < DeviceAddresses.Num(); ++AddressIndex)
 						{
 							FAutomationTestResults TestResults;
-							TestResults.GameInstance = DeviceClusterManager.GetClusterDeviceName(ClusterIndex, DeviceIndex);
+							TestResults.GameInstance = DeviceClusterManager.GetClusterGameInstance(ClusterIndex, DeviceIndex);
 							TestResults.State = EAutomationState::InProcess;
 							GameInstances.Add(TestResults.GameInstance);
 							NextTest->SetResults(ClusterIndex, CurrentTestPass, TestResults);
@@ -1124,10 +1124,11 @@ void FAutomationControllerManager::UpdateTests()
 
 					FAutomationTestResults TestResults;
 					TestResults.State = EAutomationState::Fail;
-					TestResults.GameInstance = DeviceClusterManager.GetClusterDeviceName(ClusterIndex, DeviceIndex);
-					TestResults.AddEvent(FAutomationEvent(EAutomationEventType::Error, FString::Printf(TEXT("Timeout waiting for device %s"), *TestResults.GameInstance)));
+					TestResults.GameInstance = DeviceClusterManager.GetClusterGameInstance(ClusterIndex, DeviceIndex);
+					FString DeviceName = DeviceClusterManager.GetClusterDeviceName(ClusterIndex, DeviceIndex);
+					TestResults.AddEvent(FAutomationEvent(EAutomationEventType::Error, FString::Printf(TEXT("Timeout waiting for device %s"), *DeviceName)));
 
-					UE_LOG(LogAutomationController, Error, TEXT("Removing test and marking failure after timeout waiting for device %s LastPingTime was greater than %.1f"), *TestResults.GameInstance, GameInstanceLostTimerSeconds);
+					UE_LOG(LogAutomationController, Error, TEXT("Removing test and marking failure after timeout waiting for device %s LastPingTime was greater than %.1f"), *DeviceName, GameInstanceLostTimerSeconds);
 
 					// Set the results
 					Report->SetResults(ClusterIndex, CurrentTestPass, TestResults);
@@ -1510,7 +1511,7 @@ void FAutomationControllerManager::HandleRunTestsReplyMessage(const FAutomationW
 
 		verify(DeviceClusterManager.FindDevice(Context->GetSender(), ClusterIndex, DeviceIndex));
 
-		TestResults.GameInstance = DeviceClusterManager.GetClusterDeviceName(ClusterIndex, DeviceIndex);
+		TestResults.GameInstance = DeviceClusterManager.GetClusterGameInstance(ClusterIndex, DeviceIndex);
 		TestResults.SetEvents(Message.Entries, Message.WarningTotal, Message.ErrorTotal);
 
 		// Verify this device thought it was busy
diff --git a/Engine/Source/Developer/AutomationController/Private/AutomationControllerManager.h b/Engine/Source/Developer/AutomationController/Private/AutomationControllerManager.h
--- a/Engine/Source/Developer/AutomationController/Private/AutomationControllerManager.h
+++ b/Engine/Source/Developer/AutomationController/Private/AutomationControllerManager.h
@@ -269,11 +269,16 @@ public:
 		return DeviceClusterManager.GetClusterDeviceType(ClusterIndex);
 	}
 
-	virtual FString GetGameInstanceName(const int32 ClusterIndex, const int32 DeviceIndex) const override
+	virtual FString GetDeviceName(const int32 ClusterIndex, const int32 DeviceIndex) const override
 	{
 		return DeviceClusterManager.GetClusterDeviceName(ClusterIndex, DeviceIndex);
 	}
 
+	virtual FString GetGameInstanceName(const int32 ClusterIndex, const int32 DeviceIndex) const override
+	{
+		return DeviceClusterManager.GetClusterGameInstance(ClusterIndex, DeviceIndex);
+	}
+
 	virtual void SetVisibleTestsEnabled(const bool bEnabled) override
 	{
 		ReportManager.SetVisibleTestsEnabled (bEnabled);
diff --git a/Engine/Source/Developer/AutomationController/Private/AutomationDeviceClusterManager.cpp b/Engine/Source/Developer/AutomationController/Private/AutomationDeviceClusterManager.cpp
--- a/Engine/Source/Developer/AutomationController/Private/AutomationDeviceClusterManager.cpp
+++ b/Engine/Source/Developer/AutomationController/Private/AutomationDeviceClusterManager.cpp
@@ -206,12 +206,14 @@ FString FAutomationDeviceClusterManager::GetClusterDeviceType(const int32 Cluste
 	return Clusters[ClusterIndex].DeviceTypeName;
 }
 
-
 FString FAutomationDeviceClusterManager::GetClusterDeviceName(const int32 ClusterIndex, const int32 DeviceIndex) const
 {
-	check((ClusterIndex >= 0) && (ClusterIndex < Clusters.Num()));
-	check((DeviceIndex >= 0) && (DeviceIndex < Clusters[ClusterIndex].Devices.Num()));
-	return Clusters[ClusterIndex].Devices[DeviceIndex].Info.Instance;
+	return GetDeviceInfo(ClusterIndex, DeviceIndex).DeviceName;
+}
+
+FString FAutomationDeviceClusterManager::GetClusterGameInstance(const int32 ClusterIndex, const int32 DeviceIndex) const
+{
+	return GetDeviceInfo(ClusterIndex, DeviceIndex).Instance;
 }
 
 const FAutomationDeviceInfo& FAutomationDeviceClusterManager::GetDeviceInfo(const int32 ClusterIndex, const int32 DeviceIndex) const
diff --git a/Engine/Source/Developer/AutomationController/Private/AutomationDeviceClusterManager.h b/Engine/Source/Developer/AutomationController/Private/AutomationDeviceClusterManager.h
--- a/Engine/Source/Developer/AutomationController/Private/AutomationDeviceClusterManager.h
+++ b/Engine/Source/Developer/AutomationController/Private/AutomationDeviceClusterManager.h
@@ -104,6 +104,9 @@ public:
 	/** Returns the name of a device within a cluster. */
 	FString GetClusterDeviceName(const int32 ClusterIndex, const int32 DeviceIndex) const;
 
+	/** Returns the instance of the game within a cluster. */
+	FString GetClusterGameInstance(const int32 ClusterIndex, const int32 DeviceIndex) const;
+
 	/** Returns the name of a device within a cluster. */
 	const FAutomationDeviceInfo& GetDeviceInfo(const int32 ClusterIndex, const int32 DeviceIndex) const;
 
diff --git a/Engine/Source/Developer/AutomationController/Private/AutomationReportManager.cpp b/Engine/Source/Developer/AutomationController/Private/AutomationReportManager.cpp
--- a/Engine/Source/Developer/AutomationController/Private/AutomationReportManager.cpp
+++ b/Engine/Source/Developer/AutomationController/Private/AutomationReportManager.cpp
@@ -223,7 +223,7 @@ void FAutomationReportManager::AddResultReport(TSharedPtr< IAutomationReport > I
 				}
 
 				// Create the log string
-				FString Info = FString::Printf( TEXT("DeviceName: %s, Platform: %s, Test Name: %s, Test Duration: %s, Status: %s" ),
+				FString Info = FString::Printf( TEXT("GameInstance: %s, Platform: %s, Test Name: %s, Test Duration: %s, Status: %s" ),
 					*InReport->GetGameInstanceName( ClusterIndex ),
 					*FModuleManager::GetModuleChecked< IAutomationControllerModule >( "AutomationController" ).GetAutomationController()->GetDeviceTypeName( ClusterIndex ),
 					*InReport->GetDisplayNameWithDecoration(),
diff --git a/Engine/Source/Developer/AutomationController/Public/IAutomationControllerManager.h b/Engine/Source/Developer/AutomationController/Public/IAutomationControllerManager.h
--- a/Engine/Source/Developer/AutomationController/Public/IAutomationControllerManager.h
+++ b/Engine/Source/Developer/AutomationController/Public/IAutomationControllerManager.h
@@ -271,6 +271,14 @@ public:
 	 */
 	virtual FString GetGameInstanceName(const int32 ClusterIndex, const int32 DeviceIndex) const = 0;
 
+	/**
+	 * Get a device name.
+	 *
+	 * @param ClusterIndex The cluster Index.
+	 * @param DeviceIndex The Device Index.
+	 */
+	virtual FString GetDeviceName(const int32 ClusterIndex, const int32 DeviceIndex) const = 0;
+
 	/**
 	 * Sets whether all visible tests are enabled or not.
 	 */
