Prevent approved reference images to be copied to html test report for successful test

diff --git a/Engine/Source/Developer/AutomationController/Private/AutomationControllerManager.cpp b/Engine/Source/Developer/AutomationController/Private/AutomationControllerManager.cpp
--- a/Engine/Source/Developer/AutomationController/Private/AutomationControllerManager.cpp
+++ b/Engine/Source/Developer/AutomationController/Private/AutomationControllerManager.cpp
@@ -489,15 +489,19 @@ void FAutomationControllerManager::ProcessComparisonQueue()
 				// Paths in the result are relative to the automation report directory.	
 				LocalFiles.Add(TEXT("unapproved"), FPaths::Combine(ScreenshotResultsFolder, Result.ReportIncomingFilePath));
 
-				// unapproved should always be valid. but approved/difference may be empty if this is a new screenshot
-				if (Result.ReportIncomingFilePath.Len())
+				// Don't copy reference and delta if the images are similar.
+				if (!Result.AreSimilar())
 				{
-					LocalFiles.Add(TEXT("approved"), FPaths::Combine(ScreenshotResultsFolder, Result.ReportApprovedFilePath));
-				}
+					// unapproved should always be valid. but approved/difference may be empty if this is a new screenshot
+					if (Result.ReportIncomingFilePath.Len())
+					{
+						LocalFiles.Add(TEXT("approved"), FPaths::Combine(ScreenshotResultsFolder, Result.ReportApprovedFilePath));
+					}
 
-				if (Result.ReportComparisonFilePath.Len())
-				{
-					LocalFiles.Add(TEXT("difference"), FPaths::Combine(ScreenshotResultsFolder, Result.ReportComparisonFilePath));
+					if (Result.ReportComparisonFilePath.Len())
+					{
+						LocalFiles.Add(TEXT("difference"), FPaths::Combine(ScreenshotResultsFolder, Result.ReportComparisonFilePath));
+					}
 				}
 
 				Report->AddArtifact(ClusterIndex, CurrentTestPass, FAutomationArtifact(UniqueId, Entry->ScreenshotPath, EAutomationArtifactType::Comparison, LocalFiles));
@@ -579,12 +583,12 @@ void FAutomationControllerManager::CollectTestResults(TSharedPtr<IAutomationRepo
 
 		JsonTestPassResults.TotalDuration += Results.Duration;
 
-		FString ImageCompareDirName = TEXT("imageCompare");
+		FString ArtifactDirName = TEXT("imageCompare");
 		if (FEngineVersion::Current().GetChangelist() != 0)
 		{
-			ImageCompareDirName = ImageCompareDirName / FString::FromInt(FEngineVersion::Current().GetChangelist());
+			ArtifactDirName = ArtifactDirName / FString::FromInt(FEngineVersion::Current().GetChangelist());
 		}
-		FString ImageCompareExportPath = ReportExportPath / ImageCompareDirName;
+		FString ArtifactExportPath = ReportExportPath / ArtifactDirName;
 
 		// Copy new artifacts to Report export path
 		FCriticalSection CS;
@@ -593,12 +597,14 @@ void FAutomationControllerManager::CollectTestResults(TSharedPtr<IAutomationRepo
 			TArray<FString> Keys;
 			Artifact.LocalFiles.GetKeys(Keys);
 
+			bool bOnlyUnapproved = Keys.Num() == 1 && Keys[0] == TEXT("unapproved");
+
 			ParallelFor(Keys.Num(), [&](int32 Index)
 				{
 					const FString& Key = Keys[Index];
 					FString Path = Artifact.LocalFiles[Key];
 					FPaths::MakePathRelativeTo(Path, *FPaths::AutomationReportsDir());
-					Path = ImageCompareDirName / Path;
+					Path = ArtifactDirName / Path;
 					{
 						FScopeLock Lock(&CS);
 						Artifact.Files.Add(Key, MoveTemp(Path));
@@ -606,7 +612,7 @@ void FAutomationControllerManager::CollectTestResults(TSharedPtr<IAutomationRepo
 					if (Key == TEXT("unapproved"))
 					{
 						// Copy screenshot report
-						FScreenshotExportResult ExportResult = ScreenshotManager->ExportScreenshotComparisonResult(Artifact.Name, ImageCompareExportPath);
+						FScreenshotExportResult ExportResult = ScreenshotManager->ExportScreenshotComparisonResult(Artifact.Name, ArtifactExportPath, bOnlyUnapproved);
 
 						FScopeLock Lock(&CS);
 						if (!JsonTestPassResults.ComparisonExported && ExportResult.Success)
diff --git a/Engine/Source/Developer/ScreenShotComparisonTools/Private/ScreenShotManager.cpp b/Engine/Source/Developer/ScreenShotComparisonTools/Private/ScreenShotManager.cpp
--- a/Engine/Source/Developer/ScreenShotComparisonTools/Private/ScreenShotManager.cpp
+++ b/Engine/Source/Developer/ScreenShotComparisonTools/Private/ScreenShotManager.cpp
@@ -482,7 +482,7 @@ FImageComparisonResult FScreenShotManager::CompareScreenshot(const FString& InUn
 }
 
 
-FScreenshotExportResult FScreenShotManager::ExportScreenshotComparisonResult(FString ScreenshotName, FString RootExportFolder)
+FScreenshotExportResult FScreenShotManager::ExportScreenshotComparisonResult(FString ScreenshotName, FString RootExportFolder, bool bOnlyIncoming)
 {
 	FPaths::NormalizeDirectoryName(RootExportFolder);
 
@@ -501,7 +501,9 @@ FScreenshotExportResult FScreenShotManager::ExportScreenshotComparisonResult(FSt
 		return Results;
 	}
 
-	CopyDirectory(Destination, ScreenshotResultsFolder / ScreenshotName);
+	FString Pattern = bOnlyIncoming ? TEXT("Incoming.*") : TEXT("*");
+
+	CopyDirectory(Destination, ScreenshotResultsFolder / ScreenshotName, Pattern);
 
 	Results.Success = true;
 	return Results;
@@ -556,13 +558,13 @@ FString FScreenShotManager::GetDefaultExportDirectory() const
 	return FPaths::Combine(FPaths::ProjectSavedDir(),TEXT("Exported/imageCompare"));
 }
 
-void FScreenShotManager::CopyDirectory(const FString& DestDir, const FString& SrcDir)
+void FScreenShotManager::CopyDirectory(const FString& DestDir, const FString& SrcDir, const FString& Pattern)
 {
 	TArray<FString> FilesToCopy;
 
 	FString AbsoluteSrcDir = FPaths::ConvertRelativePathToFull(SrcDir);
 
-	IFileManager::Get().FindFilesRecursive(FilesToCopy, *AbsoluteSrcDir, TEXT("*"), /*Files=*/true, /*Directories=*/false);
+	IFileManager::Get().FindFilesRecursive(FilesToCopy, *AbsoluteSrcDir, *Pattern, /*Files=*/true, /*Directories=*/false);
 
 	ParallelFor(FilesToCopy.Num(), [&](int32 Index)
 		{
diff --git a/Engine/Source/Developer/ScreenShotComparisonTools/Private/ScreenShotManager.h b/Engine/Source/Developer/ScreenShotComparisonTools/Private/ScreenShotManager.h
--- a/Engine/Source/Developer/ScreenShotComparisonTools/Private/ScreenShotManager.h
+++ b/Engine/Source/Developer/ScreenShotComparisonTools/Private/ScreenShotManager.h
@@ -47,7 +47,7 @@ public:
 
 	virtual TFuture<FImageComparisonResult> CompareScreenshotAsync(const FString& IncomingPath, const FAutomationScreenshotMetadata& MetaData, const EScreenShotCompareOptions Options) override;
 
-	virtual FScreenshotExportResult ExportScreenshotComparisonResult(FString ScreenshotName, FString ExportPath = TEXT("")) override;
+	virtual FScreenshotExportResult ExportScreenshotComparisonResult(FString ScreenshotName, FString ExportPath = TEXT(""), bool bOnlyIncoming = false) override;
 
 	virtual bool OpenComparisonReports(FString ImportPath, TArray<FComparisonReport>& OutReports) override;
 	
@@ -69,7 +69,7 @@ private:
 
 	FImageComparisonResult CompareScreenshot(const FString& IncomingPath, const FAutomationScreenshotMetadata& MetaData, const EScreenShotCompareOptions Options);
 
-	void CopyDirectory(const FString& DestDir, const FString& SrcDir);
+	void CopyDirectory(const FString& DestDir, const FString& SrcDir, const FString& Pattern);
 
 	void BuildFallbackPlatformsListFromConfig(const UScreenShotComparisonSettings* ScreenShotSettings);
 
diff --git a/Engine/Source/Developer/ScreenShotComparisonTools/Public/Interfaces/IScreenShotManager.h b/Engine/Source/Developer/ScreenShotComparisonTools/Public/Interfaces/IScreenShotManager.h
--- a/Engine/Source/Developer/ScreenShotComparisonTools/Public/Interfaces/IScreenShotManager.h
+++ b/Engine/Source/Developer/ScreenShotComparisonTools/Public/Interfaces/IScreenShotManager.h
@@ -68,7 +68,7 @@ public:
 	/**
 	 * Exports target screenshot report to the export location specified
 	 */
-	virtual FScreenshotExportResult ExportScreenshotComparisonResult(FString ScreenshotName, FString ExportPath = TEXT("")) = 0;
+	virtual FScreenshotExportResult ExportScreenshotComparisonResult(FString ScreenshotName, FString ExportPath = TEXT(""), bool bOnlyIncoming = false) = 0;
 
 	/**
 	 * Imports screenshot comparison data from a given path.
