Fix include path inconsistentcy with Launch module when -AllModules is passed

#rnx

diff --git a/Engine/Source/Developer/AutomationController/AutomationController.Build.cs b/Engine/Source/Developer/AutomationController/AutomationController.Build.cs
--- a/Engine/Source/Developer/AutomationController/AutomationController.Build.cs
+++ b/Engine/Source/Developer/AutomationController/AutomationController.Build.cs
@@ -41,13 +41,6 @@ namespace UnrealBuildTool.Rules
                 );
             }
 
-			PublicIncludePathModuleNames.AddRange(
-				new string[]
-				{
-					"AutomationTest"
-				}
-			);
-
 			PrivateIncludePathModuleNames.AddRange(
 				new string[]
 				{
diff --git a/Engine/Source/Runtime/ClientPilot/ClientPilot.Build.cs b/Engine/Source/Runtime/ClientPilot/ClientPilot.Build.cs
--- a/Engine/Source/Runtime/ClientPilot/ClientPilot.Build.cs
+++ b/Engine/Source/Runtime/ClientPilot/ClientPilot.Build.cs
@@ -15,7 +15,8 @@ public class ClientPilot : ModuleRules
         );
 
 		PublicIncludePathModuleNames.AddRange(new string[] {
-				"AutomationController"
+				"AutomationController",
+				"AutomationTest",
 		});
 
 		if (Target.bCompileAgainstEngine && Target.Configuration != UnrealTargetConfiguration.Shipping)
diff --git a/Engine/Source/Runtime/Launch/Launch.Build.cs b/Engine/Source/Runtime/Launch/Launch.Build.cs
--- a/Engine/Source/Runtime/Launch/Launch.Build.cs
+++ b/Engine/Source/Runtime/Launch/Launch.Build.cs
@@ -9,6 +9,7 @@ public class Launch : ModuleRules
 	{
 		PrivateIncludePathModuleNames.AddRange(new string[] {
 				"AutomationController",
+				"AutomationTest",
 				"ProfileVisualizer",
 		});
 
