Added EComparisonFileTypes enum for IAutomationReport.h, changed TEXT() to enum AutomationControllerManager.cpp when referencing test report Artifact file key

Use enum or predefined const strings when referencing test report Artifact file key in FAutomationArtifact

diff --git a/Engine/Source/Developer/AutomationController/Private/AutomationControllerManager.cpp b/Engine/Source/Developer/AutomationController/Private/AutomationControllerManager.cpp
--- a/Engine/Source/Developer/AutomationController/Private/AutomationControllerManager.cpp
+++ b/Engine/Source/Developer/AutomationController/Private/AutomationControllerManager.cpp
@@ -482,12 +482,12 @@ void FAutomationControllerManager::ProcessComparisonQueue()
 			if (Report.IsValid())
 			{
 				// Record the artifacts for the test.
-				TMap<FString, FString> LocalFiles;
+				TMap<EComparisonFileTypes, FString> LocalFiles;
 
 				FString ScreenshotResultsFolder = FPaths::AutomationReportsDir();
 	
 				// Paths in the result are relative to the automation report directory.	
-				LocalFiles.Add(TEXT("unapproved"), FPaths::Combine(ScreenshotResultsFolder, Result.ReportIncomingFilePath));
+				LocalFiles.Add(EComparisonFileTypes::Unapproved, FPaths::Combine(ScreenshotResultsFolder, Result.ReportIncomingFilePath));
 
 				// Don't copy reference and delta if the images are similar.
 				if (!Result.AreSimilar())
@@ -495,12 +495,12 @@ void FAutomationControllerManager::ProcessComparisonQueue()
 					// unapproved should always be valid. but approved/difference may be empty if this is a new screenshot
 					if (Result.ReportIncomingFilePath.Len())
 					{
-						LocalFiles.Add(TEXT("approved"), FPaths::Combine(ScreenshotResultsFolder, Result.ReportApprovedFilePath));
+						LocalFiles.Add(EComparisonFileTypes::Approved, FPaths::Combine(ScreenshotResultsFolder, Result.ReportApprovedFilePath));
 					}
 
 					if (Result.ReportComparisonFilePath.Len())
 					{
-						LocalFiles.Add(TEXT("difference"), FPaths::Combine(ScreenshotResultsFolder, Result.ReportComparisonFilePath));
+						LocalFiles.Add(EComparisonFileTypes::Difference, FPaths::Combine(ScreenshotResultsFolder, Result.ReportComparisonFilePath));
 					}
 				}
 
@@ -594,14 +594,14 @@ void FAutomationControllerManager::CollectTestResults(TSharedPtr<IAutomationRepo
 		FCriticalSection CS;
 		for (FAutomationArtifact& Artifact : TestResult.GetArtifacts())
 		{
-			TArray<FString> Keys;
+			TArray<EComparisonFileTypes> Keys;
 			Artifact.LocalFiles.GetKeys(Keys);
 
-			bool bOnlyUnapproved = Keys.Num() == 1 && Keys[0] == TEXT("unapproved");
+			bool bOnlyUnapproved = Keys.Num() == 1 && Keys[0] == EComparisonFileTypes::Unapproved;
 
 			ParallelFor(Keys.Num(), [&](int32 Index)
 				{
-					const FString& Key = Keys[Index];
+					const EComparisonFileTypes& Key = Keys[Index];
 					FString Path = Artifact.LocalFiles[Key];
 					FPaths::MakePathRelativeTo(Path, *FPaths::AutomationReportsDir());
 					Path = ArtifactDirName / Path;
@@ -609,7 +609,7 @@ void FAutomationControllerManager::CollectTestResults(TSharedPtr<IAutomationRepo
 						FScopeLock Lock(&CS);
 						Artifact.Files.Add(Key, MoveTemp(Path));
 					}
-					if (Key == TEXT("unapproved"))
+					if (Key == EComparisonFileTypes::Unapproved)
 					{
 						// Copy screenshot report
 						FScreenshotExportResult ExportResult = ScreenshotManager->ExportScreenshotComparisonResult(Artifact.Name, ArtifactExportPath, bOnlyUnapproved);
diff --git a/Engine/Source/Developer/AutomationController/Public/IAutomationReport.h b/Engine/Source/Developer/AutomationController/Public/IAutomationReport.h
--- a/Engine/Source/Developer/AutomationController/Public/IAutomationReport.h
+++ b/Engine/Source/Developer/AutomationController/Public/IAutomationReport.h
@@ -31,6 +31,14 @@ enum class EAutomationArtifactType : uint8
 	Comparison
 };
 
+UENUM()
+enum class EComparisonFileTypes : uint8
+{
+	Unapproved,
+	Approved,
+	Difference
+};
+
 USTRUCT()
 struct FAutomationArtifact
 {
@@ -43,7 +51,7 @@ public:
 	{
 	}
 
-	FAutomationArtifact(FGuid InUniqueId, const FString& InName, EAutomationArtifactType InType, const TMap<FString, FString>& InLocalFiles)
+	FAutomationArtifact(FGuid InUniqueId, const FString& InName, EAutomationArtifactType InType, const TMap<EComparisonFileTypes, FString>& InLocalFiles)
 		: Id(InUniqueId)
 		, Name(InName)
 		, Type(InType)
@@ -63,11 +71,12 @@ public:
 	EAutomationArtifactType Type;
 
 	UPROPERTY()
-	TMap<FString, FString> Files;
+	TMap<EComparisonFileTypes, FString> Files;
 
 	// Local Files are the files generated during a testing run, once exported, the individual file paths
 	// should be stored in the Files map.
-	TMap<FString, FString> LocalFiles;
+
+	TMap<EComparisonFileTypes, FString> LocalFiles;
 };
 
 /**
