Avoid changing the AutomationController state when running test while receiving multiple FindWorkers responses

diff --git a/Engine/Source/Developer/AutomationController/Private/AutomationControllerManager.cpp b/Engine/Source/Developer/AutomationController/Private/AutomationControllerManager.cpp
--- a/Engine/Source/Developer/AutomationController/Private/AutomationControllerManager.cpp
+++ b/Engine/Source/Developer/AutomationController/Private/AutomationControllerManager.cpp
@@ -1225,6 +1225,12 @@ void FAutomationControllerManager::HandleFindWorkersResponseMessage(const FAutom
 		DeviceClusterManager.AddDeviceFromMessage(Context->GetSender(), Message, DeviceGroupFlags);
 	}
 
+	if (AutomationTestState == EAutomationControllerModuleState::Running)
+	{
+		UE_LOG(LogAutomationController, Verbose, TEXT("Already running test. Don't request tests from %s"), *Context->GetSender().ToString());
+		return;
+	}
+
 	RequestTests();
 
 	SetControllerStatus(EAutomationControllerModuleState::Ready);
