Add header and remove redundant information from exported test results (csv).

diff --git a/Engine/Source/Developer/AutomationController/Private/AutomationReportManager.cpp b/Engine/Source/Developer/AutomationController/Private/AutomationReportManager.cpp
--- a/Engine/Source/Developer/AutomationController/Private/AutomationReportManager.cpp
+++ b/Engine/Source/Developer/AutomationController/Private/AutomationReportManager.cpp
@@ -7,6 +7,18 @@
 #include "IAutomationControllerModule.h"
 #include "AutomationReport.h"
 
+namespace
+{
+	// Note that if you change at least one line here you must change other lines in order to keep consistent columns count.
+	// Header string to add as a report's first line
+	constexpr TCHAR ReportHeader[] = TEXT("\"GameInstance\",\"Platform\",\"Test Name\",\"Duration\",\"Status\"");
+	// Format string that is used to fill the lines that contain report's data
+	constexpr TCHAR ReportLineFormatString[] = TEXT("\"%s\",\"%s\",\"%s\",\"%s\",\"%s\"");
+	// Pre and post fix strings to add to the errors and warnings for viewing in CSV. Spaces correctly, and allows commas in the message
+	constexpr TCHAR ErrorReportLinePrefixText[] = TEXT("\t,,,,\"");
+	constexpr TCHAR ErrorReportLineSufixText[] = TEXT("\"");
+} // anonymous namespace
+
 
 FAutomationReportManager::FAutomationReportManager()
 {
@@ -105,7 +117,7 @@ void FAutomationReportManager::SetEnabledTests(const TArray<FString>& EnabledTes
 const bool FAutomationReportManager::ExportReport(uint32 FileExportTypeMask, const int32 NumDeviceClusters)
 {
 	// The results log. Errors, warnings etc are added to this
-	TArray<FString> ResultsLog;
+	TArray<FString> ResultsLog = { ReportHeader };
 
 	// Create the report be recursively going through the results tree
 	FindLeafReport(ReportRoot, NumDeviceClusters, ResultsLog, FileExportTypeMask);
@@ -113,8 +125,8 @@ const bool FAutomationReportManager::ExportReport(uint32 FileExportTypeMask, con
 	// Has a report been generated
 	bool ReportGenerated = false;
 
-	// Ensure we have a log to write
-	if (ResultsLog.Num())
+	// Ensure we have a log to write (the first line is a header line)
+	if (ResultsLog.Num() > 1)
 	{
 		// Create the file name
 		FString FileName = FString::Printf( TEXT( "Automation%s.csv" ), *FDateTime::Now().ToString() );
@@ -223,7 +235,7 @@ void FAutomationReportManager::AddResultReport(TSharedPtr< IAutomationReport > I
 				}
 
 				// Create the log string
-				FString Info = FString::Printf( TEXT("GameInstance: %s, Platform: %s, Test Name: %s, Test Duration: %s, Status: %s" ),
+				FString Info = FString::Printf( ReportLineFormatString,
 					*InReport->GetGameInstanceName( ClusterIndex ),
 					*FModuleManager::GetModuleChecked< IAutomationControllerModule >( "AutomationController" ).GetAutomationController()->GetDeviceTypeName( ClusterIndex ),
 					*InReport->GetDisplayNameWithDecoration(),
@@ -232,11 +244,6 @@ void FAutomationReportManager::AddResultReport(TSharedPtr< IAutomationReport > I
 
 				ResultsLog.Add( Info );
 
-				// Pre and post fix strings to add to the errors and warnings for viewing in CSV. Spaces correctly, and allows commas in the message
-				const FString ErrorPrefixText = TEXT( "\t,,,\"" );
-				const FString ErrorSufixText = TEXT( "\"" );
-
-
 				// Add any warning or errors
 				for ( const FAutomationExecutionEntry& Entry : TestResults.GetEntries() )
 				{
@@ -245,19 +252,19 @@ void FAutomationReportManager::AddResultReport(TSharedPtr< IAutomationReport > I
 					case EAutomationEventType::Info:
 						if ( EFileExportType::IsSet(FileExportTypeMask, EFileExportType::FET_All) || EFileExportType::IsSet(FileExportTypeMask, EFileExportType::FET_Logs) )
 						{
-							ResultsLog.Add(ErrorPrefixText + Entry.ToString() + ErrorSufixText);
+							ResultsLog.Add(ErrorReportLinePrefixText + Entry.ToString() + ErrorReportLineSufixText);
 						}
 						break;
 					case EAutomationEventType::Warning:
 						if ( EFileExportType::IsSet(FileExportTypeMask, EFileExportType::FET_All) || EFileExportType::IsSet(FileExportTypeMask, EFileExportType::FET_Warnings) )
 						{
-							ResultsLog.Add(ErrorPrefixText + Entry.ToString() + ErrorSufixText);
+							ResultsLog.Add(ErrorReportLinePrefixText + Entry.ToString() + ErrorReportLineSufixText);
 						}
 						break;
 					case EAutomationEventType::Error:
 						if ( EFileExportType::IsSet(FileExportTypeMask, EFileExportType::FET_All) || EFileExportType::IsSet(FileExportTypeMask, EFileExportType::FET_Errors) )
 						{
-							ResultsLog.Add(ErrorPrefixText + Entry.ToString() + ErrorSufixText);
+							ResultsLog.Add(ErrorReportLinePrefixText + Entry.ToString() + ErrorReportLineSufixText);
 						}
 						break;
 					}
