Remove extra characters from skipped and warning status in local csv test report

#rnx

diff --git a/Engine/Source/Developer/AutomationController/Private/AutomationReportManager.cpp b/Engine/Source/Developer/AutomationController/Private/AutomationReportManager.cpp
--- a/Engine/Source/Developer/AutomationController/Private/AutomationReportManager.cpp
+++ b/Engine/Source/Developer/AutomationController/Private/AutomationReportManager.cpp
@@ -219,14 +219,14 @@ void FAutomationReportManager::AddResultReport(TSharedPtr< IAutomationReport > I
 				}
 				else if( InReport->GetState( ClusterIndex,CurrentTestPass ) == EAutomationState::Skipped )
 				{
-					Status = TEXT( "Skipped." );
+					Status = TEXT( "Skipped" );
 				}
 				else
 				{
 					Status = "Issues ";
 					if (TestResults.GetWarningTotal())
 					{
-						Status += "Warnings ";
+						Status += "Warnings";
 					}
 					if (TestResults.GetErrorTotal())
 					{
