Automation test - standardize status in csv report

#rnx

diff --git a/Engine/Source/Developer/AutomationController/Private/AutomationReportManager.cpp b/Engine/Source/Developer/AutomationController/Private/AutomationReportManager.cpp
--- a/Engine/Source/Developer/AutomationController/Private/AutomationReportManager.cpp
+++ b/Engine/Source/Developer/AutomationController/Private/AutomationReportManager.cpp
@@ -210,29 +210,7 @@ void FAutomationReportManager::AddResultReport(TSharedPtr< IAutomationReport > I
 				}
 
 				// Build a status string that contains information about the test 
-				FString Status;
-
-				// Was the test a success
-				if ( ( TestResults.GetWarningTotal() == 0 ) && ( TestResults.GetErrorTotal() == 0 ) && ( InReport->GetState(ClusterIndex, CurrentTestPass) == EAutomationState::Success ) )
-				{
-					Status = "Success";
-				}
-				else if( InReport->GetState( ClusterIndex,CurrentTestPass ) == EAutomationState::Skipped )
-				{
-					Status = TEXT( "Skipped" );
-				}
-				else
-				{
-					Status = "Issues ";
-					if (TestResults.GetWarningTotal())
-					{
-						Status += "Warnings";
-					}
-					if (TestResults.GetErrorTotal())
-					{
-						Status += "Errors";
-					}
-				}
+				FString Status = AutomationStateToString(TestResults.State);
 
 				// Create the log string
 				FString Info = FString::Printf( ReportLineFormatString,
