Fix ExecCmds Automation queueing command line to support input like -exccmds="Automation RunTest MyTest,Automation List,Automation Quit"

#rnx

diff --git a/Engine/Source/Developer/AutomationController/Private/AutomationCommandline.cpp b/Engine/Source/Developer/AutomationController/Private/AutomationCommandline.cpp
--- a/Engine/Source/Developer/AutomationController/Private/AutomationCommandline.cpp
+++ b/Engine/Source/Developer/AutomationController/Private/AutomationCommandline.cpp
@@ -451,6 +451,13 @@ public:
 					// We have finished the testing, and results are available
 					AutomationTestState = EAutomationTestState::Complete;
 				}
+				else if (!IsAboutToRunTest())
+				{
+					// Register for the callback that tells us there are tests available
+					if (!TestsRefreshedHandle.IsValid()) {
+						TestsRefreshedHandle = AutomationController->OnTestsRefreshed().AddRaw(this, &FAutomationExecCmd::HandleRefreshTestCallback);
+					}
+				}
 				break;
 			}
 		}
@@ -478,6 +485,13 @@ public:
 
 		return false;
 	}
+
+	bool IsAboutToRunTest()
+	{
+		return (AutomationCommand == EAutomationCommand::RunCommandLineTests
+			|| AutomationCommand == EAutomationCommand::RunAll
+			|| AutomationCommand == EAutomationCommand::RunFilter);
+	}
 	
 protected:
 	/** Console commands, see embeded usage statement **/
@@ -490,11 +504,6 @@ protected:
 		// figure out if we are handling this request
 		if (FParse::Command(&Cmd, TEXT("Automation")))
 		{
-			if (AutomationTestState != EAutomationTestState::Idle)
-			{
-				Ar.Logf(TEXT("Automation: Skipping this request because already executing tests!"));
-				return false;
-			}
 
 			TArray<FString> CommandList;
 			FString(Cmd).ParseIntoArray(CommandList, TEXT(";"), true);
