Fix infinite loop when Quit command is input more than once

#rnx

diff --git a/Engine/Source/Developer/AutomationController/Private/AutomationCommandline.cpp b/Engine/Source/Developer/AutomationController/Private/AutomationCommandline.cpp
--- a/Engine/Source/Developer/AutomationController/Private/AutomationCommandline.cpp
+++ b/Engine/Source/Developer/AutomationController/Private/AutomationCommandline.cpp
@@ -423,7 +423,7 @@ public:
 					AutomationCommandQueue.RemoveAt(0);
 					if (AutomationCommand == EAutomationCommand::Quit || AutomationCommand == EAutomationCommand::SoftQuit)
 					{
-						if (AutomationCommandQueue.IsValidIndex(0))
+						if (AutomationCommandQueue.IsValidIndex(0) && !IsQuitQueued())
 						{
 							// Add Quit and SoftQuit commands back to the end of the array.
 							AutomationCommandQueue.Add(AutomationCommand);
@@ -492,6 +492,20 @@ public:
 			|| AutomationCommand == EAutomationCommand::RunAll
 			|| AutomationCommand == EAutomationCommand::RunFilter);
 	}
+
+	bool IsQuitQueued()
+	{
+		for (auto Command : AutomationCommandQueue)
+		{
+			if (Command == EAutomationCommand::Quit
+				|| Command == EAutomationCommand::SoftQuit)
+			{
+				return true;
+			}
+		}
+
+		return false;
+	}
 	
 protected:
 	/** Console commands, see embeded usage statement **/
@@ -651,11 +665,21 @@ protected:
 				}
 				else if (FParse::Command(&TempCmd, TEXT("Quit")))
 				{
+					if (IsQuitQueued())
+					{
+						Ar.Log(TEXT("Automation: Quit command is already Queued."));
+						continue;
+					}
 					AutomationCommandQueue.Add(EAutomationCommand::Quit);
 					Ar.Logf(TEXT("Automation: Quit Command Queued."));
 				}
 				else if (FParse::Command(&TempCmd, TEXT("SoftQuit")))
 				{
+					if (IsQuitQueued())
+					{
+						Ar.Log(TEXT("Automation: Quit command is already Queued."));
+						continue;
+					}
 					AutomationCommandQueue.Add(EAutomationCommand::SoftQuit);
 					Ar.Logf(TEXT("Automation: SoftQuit Command Queued."));
 				}
