Fix missing deviceInstance in test results when running without ResumeRunTests flag

#rnx

diff --git a/Engine/Source/Developer/AutomationController/Private/AutomationControllerManager.cpp b/Engine/Source/Developer/AutomationController/Private/AutomationControllerManager.cpp
--- a/Engine/Source/Developer/AutomationController/Private/AutomationControllerManager.cpp
+++ b/Engine/Source/Developer/AutomationController/Private/AutomationControllerManager.cpp
@@ -580,6 +580,10 @@ void FAutomationControllerManager::CollectTestResults(TSharedPtr<IAutomationRepo
 		TestResult.SetEvents(Results.GetEntries(), Results.GetWarningTotal(), Results.GetErrorTotal());
 		TestResult.SetArtifacts(Results.Artifacts);
 		TestResult.Duration = Results.Duration;
+		if (TestResult.DeviceInstance.IsEmpty())
+		{
+			TestResult.DeviceInstance = { Results.GameInstance };
+		}
 
 		JsonTestPassResults.TotalDuration += Results.Duration;
 
