Implement OnBeforeAllTestsEvent and OnAfterAllTestsEvent in FAutomationTestFramework

diff --git a/Engine/Plugins/Tests/TestSamples/Source/TestSamples/Private/TestSamples.cpp b/Engine/Plugins/Tests/TestSamples/Source/TestSamples/Private/TestSamples.cpp
--- a/Engine/Plugins/Tests/TestSamples/Source/TestSamples/Private/TestSamples.cpp
+++ b/Engine/Plugins/Tests/TestSamples/Source/TestSamples/Private/TestSamples.cpp
@@ -10,6 +10,10 @@ DEFINE_LOG_CATEGORY_STATIC(LogTestSamples, Log, All);
 
 void FTestSamplesModule::StartupModule()
 {
+	// Hook callbacks on before all tests Start and End event
+	FAutomationTestFramework::Get().OnBeforeAllTestsEvent.AddRaw(this, &FTestSamplesModule::OnBeforeAllTests);
+	FAutomationTestFramework::Get().OnAfterAllTestsEvent.AddRaw(this, &FTestSamplesModule::OnAfterAllTests);
+
 	// Hook callbacks to Test Start and End event
 	FAutomationTestFramework::Get().OnTestStartEvent.AddRaw(this, &FTestSamplesModule::OnTestStart);
 	FAutomationTestFramework::Get().OnTestEndEvent.AddRaw(this, &FTestSamplesModule::OnTestEnd);
@@ -20,6 +24,9 @@ void FTestSamplesModule::ShutdownModule()
 	// Remove all registered callbacks on shutdown
 	FAutomationTestFramework::Get().OnTestStartEvent.RemoveAll(this);
 	FAutomationTestFramework::Get().OnTestEndEvent.RemoveAll(this);
+
+	FAutomationTestFramework::Get().OnBeforeAllTestsEvent.RemoveAll(this);
+	FAutomationTestFramework::Get().OnAfterAllTestsEvent.RemoveAll(this);
 }
 
 void FTestSamplesModule::OnTestStart(FAutomationTestBase* Test)
@@ -38,6 +45,16 @@ void FTestSamplesModule::OnTestEnd(FAutomationTestBase* Test)
 
 }
 
+void FTestSamplesModule::OnBeforeAllTests()
+{
+	UE_LOG(LogTestSamples, Verbose, TEXT("Start running tests"));
+}
+
+void FTestSamplesModule::OnAfterAllTests()
+{
+	UE_LOG(LogTestSamples, Verbose, TEXT("Running tests completed"));
+}
+
 #undef LOCTEXT_NAMESPACE
 	
 IMPLEMENT_MODULE(FTestSamplesModule, TestSamples)
diff --git a/Engine/Plugins/Tests/TestSamples/Source/TestSamples/Public/TestSamples.h b/Engine/Plugins/Tests/TestSamples/Source/TestSamples/Public/TestSamples.h
--- a/Engine/Plugins/Tests/TestSamples/Source/TestSamples/Public/TestSamples.h
+++ b/Engine/Plugins/Tests/TestSamples/Source/TestSamples/Public/TestSamples.h
@@ -15,7 +15,9 @@ public:
 	virtual void ShutdownModule() override;
 
 private:
+	void OnBeforeAllTests();
+	void OnAfterAllTests();
+
 	void OnTestStart(FAutomationTestBase* Test);
 	void OnTestEnd(FAutomationTestBase* Test);
-
 };
diff --git a/Engine/Source/Developer/AutomationController/Private/AutomationControllerManager.cpp b/Engine/Source/Developer/AutomationController/Private/AutomationControllerManager.cpp
--- a/Engine/Source/Developer/AutomationController/Private/AutomationControllerManager.cpp
+++ b/Engine/Source/Developer/AutomationController/Private/AutomationControllerManager.cpp
@@ -345,11 +345,12 @@ void FAutomationControllerManager::RunTests(const bool bInIsLocalSession)
 			//mark the device as idle
 			DeviceClusterManager.SetTest(ClusterIndex, DeviceIndex, NULL);
 
-			// Send command to reset tests (delete local files, etc)
+			// Send command to start test session and reset tests (delete local files, etc)
 			FMessageAddress MessageAddress = DeviceClusterManager.GetDeviceMessageAddress(ClusterIndex, DeviceIndex);
-			UE_LOG(LogAutomationController, Log, TEXT("Sending Reset Tests to %s"), *MessageAddress.ToString());
+			UE_LOG(LogAutomationController, Log, TEXT("Sending StartTestSession to %s"), *MessageAddress.ToString());
 
-			MessageEndpoint->Send(FMessageEndpoint::MakeMessage<FAutomationWorkerResetTests>(), MessageAddress);
+			MessageEndpoint->Send(FMessageEndpoint::MakeMessage<FAutomationWorkerStartTestSession>(), MessageAddress);
+			StartedTestSessionAddressSet.Add(MessageAddress);
 
 			// Store devices info into the json report.
 			if (JsonTestPassResults.IsRequired)
@@ -408,6 +409,13 @@ void FAutomationControllerManager::StopTests()
 
 	TestRunningArray.Empty();
 
+	for (const auto& MessageAddress : StartedTestSessionAddressSet)
+	{
+		UE_LOG(LogAutomationController, Log, TEXT("Sending StopTestSession to %s"), *MessageAddress.ToString());
+		MessageEndpoint->Send(FMessageEndpoint::MakeMessage<FAutomationWorkerStopTestSession>(), MessageAddress);
+	}
+	StartedTestSessionAddressSet.Empty();
+
 	// Close play window
 #if WITH_EDITOR
 	if (GUnrealEd && !bKeepPIEOpen)
@@ -548,6 +556,14 @@ void FAutomationControllerManager::ProcessAvailableTasks()
 					}
 					#endif
 
+					// Notify the workers about stopping the session.
+					for (const auto& MessageAddress : StartedTestSessionAddressSet)
+					{
+						UE_LOG(LogAutomationController, Log, TEXT("Sending StopTestSession to %s"), *MessageAddress.ToString());
+						MessageEndpoint->Send(FMessageEndpoint::MakeMessage<FAutomationWorkerStopTestSession>(), MessageAddress);
+					}
+					StartedTestSessionAddressSet.Empty();
+
 					//Notify the graphical layout we are done processing results.
 					TestsCompleteDelegate.Broadcast();
 				}
diff --git a/Engine/Source/Developer/AutomationController/Private/AutomationControllerManager.h b/Engine/Source/Developer/AutomationController/Private/AutomationControllerManager.h
--- a/Engine/Source/Developer/AutomationController/Private/AutomationControllerManager.h
+++ b/Engine/Source/Developer/AutomationController/Private/AutomationControllerManager.h
@@ -595,6 +595,9 @@ private:
 	/** A array of running tests. */
 	TArray< FTestRunningInfo > TestRunningArray;
 
+	/** Set of addresses of started sessions. */
+	TSet< FMessageAddress > StartedTestSessionAddressSet;
+
 	/** The number of test passes to perform. */
 	int32 NumTestPasses = 0;
 
diff --git a/Engine/Source/Runtime/AutomationMessages/Public/AutomationWorkerMessages.h b/Engine/Source/Runtime/AutomationMessages/Public/AutomationWorkerMessages.h
--- a/Engine/Source/Runtime/AutomationMessages/Public/AutomationWorkerMessages.h
+++ b/Engine/Source/Runtime/AutomationMessages/Public/AutomationWorkerMessages.h
@@ -131,7 +131,16 @@ struct FAutomationWorkerPing
 /**
  */
 USTRUCT()
-struct FAutomationWorkerResetTests
+struct FAutomationWorkerStartTestSession
+{
+	GENERATED_USTRUCT_BODY()
+};
+
+
+/**
+ */
+USTRUCT()
+struct FAutomationWorkerStopTestSession
 {
 	GENERATED_USTRUCT_BODY()
 };
diff --git a/Engine/Source/Runtime/AutomationWorker/Private/AutomationWorkerModule.cpp b/Engine/Source/Runtime/AutomationWorker/Private/AutomationWorkerModule.cpp
--- a/Engine/Source/Runtime/AutomationWorker/Private/AutomationWorkerModule.cpp
+++ b/Engine/Source/Runtime/AutomationWorker/Private/AutomationWorkerModule.cpp
@@ -140,7 +140,8 @@ void FAutomationWorkerModule::Initialize()
 			.Handling<FAutomationWorkerFindWorkers>(this, &FAutomationWorkerModule::HandleFindWorkersMessage)
 			.Handling<FAutomationWorkerNextNetworkCommandReply>(this, &FAutomationWorkerModule::HandleNextNetworkCommandReplyMessage)
 			.Handling<FAutomationWorkerPing>(this, &FAutomationWorkerModule::HandlePingMessage)
-			.Handling<FAutomationWorkerResetTests>(this, &FAutomationWorkerModule::HandleResetTests)
+			.Handling<FAutomationWorkerStartTestSession>(this, &FAutomationWorkerModule::HandleStartTestSession)
+			.Handling<FAutomationWorkerStopTestSession>(this, &FAutomationWorkerModule::HandleStopTestSession)
 			.Handling<FAutomationWorkerRequestTests>(this, &FAutomationWorkerModule::HandleRequestTestsMessage)
 			.Handling<FAutomationWorkerRunTests>(this, &FAutomationWorkerModule::HandleRunTestsMessage)
 			.Handling<FAutomationWorkerImageComparisonResults>(this, &FAutomationWorkerModule::HandleScreenShotCompared)
@@ -339,11 +340,19 @@ void FAutomationWorkerModule::HandlePingMessage( const FAutomationWorkerPing& Me
 }
 
 
-void FAutomationWorkerModule::HandleResetTests( const FAutomationWorkerResetTests& Message, const TSharedRef<IMessageContext, ESPMode::ThreadSafe>& Context )
+void FAutomationWorkerModule::HandleStartTestSession( const FAutomationWorkerStartTestSession& Message, const TSharedRef<IMessageContext, ESPMode::ThreadSafe>& Context )
 {
-	UE_LOG(LogAutomationWorker, Log, TEXT("Received ResetTests from %s"), *Context->GetSender().ToString());
+	UE_LOG(LogAutomationWorker, Log, TEXT("Received StartTestSession from %s"), *Context->GetSender().ToString());
 
 	FAutomationTestFramework::Get().ResetTests();
+	FAutomationTestFramework::Get().OnBeforeAllTestsEvent.Broadcast();
+}
+
+void FAutomationWorkerModule::HandleStopTestSession(const FAutomationWorkerStopTestSession& Message, const TSharedRef<IMessageContext, ESPMode::ThreadSafe>& Context)
+{
+	UE_LOG(LogAutomationWorker, Log, TEXT("Received StopTestSession from %s"), *Context->GetSender().ToString());
+
+	FAutomationTestFramework::Get().OnAfterAllTestsEvent.Broadcast();
 }
 
 
diff --git a/Engine/Source/Runtime/AutomationWorker/Private/AutomationWorkerModule.h b/Engine/Source/Runtime/AutomationWorker/Private/AutomationWorkerModule.h
--- a/Engine/Source/Runtime/AutomationWorker/Private/AutomationWorkerModule.h
+++ b/Engine/Source/Runtime/AutomationWorker/Private/AutomationWorkerModule.h
@@ -19,7 +19,8 @@ struct FAutomationWorkerPerformanceDataResponse;
 struct FAutomationWorkerNextNetworkCommandReply;
 struct FAutomationWorkerPing;
 struct FAutomationWorkerRequestTests;
-struct FAutomationWorkerResetTests;
+struct FAutomationWorkerStartTestSession;
+struct FAutomationWorkerStopTestSession;
 struct FAutomationWorkerRunTests;
 struct FMessageAddress;
 
@@ -85,8 +86,11 @@ private:
 	/** Handles FAutomationWorkerPing messages. */
 	void HandlePingMessage(const FAutomationWorkerPing& Message, const TSharedRef<IMessageContext, ESPMode::ThreadSafe>& Context);
 
-	/** Handles FAutomationWorkerResetTests messages. */
-	void HandleResetTests(const FAutomationWorkerResetTests& Message, const TSharedRef<IMessageContext, ESPMode::ThreadSafe>& Context);
+	/** Handles FAutomationWorkerStartTestSession messages. */
+	void HandleStartTestSession(const FAutomationWorkerStartTestSession& Message, const TSharedRef<IMessageContext, ESPMode::ThreadSafe>& Context);
+
+	/** Handles FAutomationWorkerStopTestSession messages. */
+	void HandleStopTestSession(const FAutomationWorkerStopTestSession& Message, const TSharedRef<IMessageContext, ESPMode::ThreadSafe>& Context);
 
 	/** Handles FAutomationWorkerRequestTests messages. */
 	void HandleRequestTestsMessage(const FAutomationWorkerRequestTests& Message, const TSharedRef<IMessageContext, ESPMode::ThreadSafe>& Context);
diff --git a/Engine/Source/Runtime/Core/Public/Misc/AutomationTest.h b/Engine/Source/Runtime/Core/Public/Misc/AutomationTest.h
--- a/Engine/Source/Runtime/Core/Public/Misc/AutomationTest.h
+++ b/Engine/Source/Runtime/Core/Public/Misc/AutomationTest.h
@@ -841,6 +841,12 @@ public:
 	/** The final call related to screenshots, after they've been taken, and after they've been compared (or not if automation isn't running). */
 	FSimpleMulticastDelegate OnScreenshotTakenAndCompared;
 
+	/** Called before all chosen tests run. */
+	FSimpleMulticastDelegate OnBeforeAllTestsEvent;
+
+	/** Called after all chosen tests run have finished. */
+	FSimpleMulticastDelegate OnAfterAllTestsEvent;
+
 	/**
 	 * Return the singleton instance of the framework.
 	 *
