Avoid running all tests if Group filter does not exist

#rnx

diff --git a/Engine/Source/Developer/AutomationController/Private/AutomationCommandline.cpp b/Engine/Source/Developer/AutomationController/Private/AutomationCommandline.cpp
--- a/Engine/Source/Developer/AutomationController/Private/AutomationCommandline.cpp
+++ b/Engine/Source/Developer/AutomationController/Private/AutomationCommandline.cpp
@@ -186,7 +186,7 @@ public:
 
 				if (!FoundGroup)
 				{
-					UE_LOG(LogAutomationCommandLine, Warning, TEXT("No matching group named %s"), *GroupName);
+					UE_LOG(LogAutomationCommandLine, Error, TEXT("No matching group named %s"), *GroupName);
 				}
 			}			
 			else
@@ -209,13 +209,16 @@ public:
 			}
 		}
 
-		FilterAny->SetFilters(FiltersList);
-		InFilters->Add(MakeShareable(FilterAny));
+		if (!FiltersList.IsEmpty())
+		{
+			FilterAny->SetFilters(FiltersList);
+			InFilters->Add(MakeShareable(FilterAny));
 
-		// SetFilter applies all filters from the AutomationFilters array
-		AutomationController->SetFilter(InFilters);
-		// Fill OutFilteredTestNames array with filtered test names
-		AutomationController->GetFilteredTestNames(OutFilteredTestNames);
+			// SetFilter applies all filters from the AutomationFilters array
+			AutomationController->SetFilter(InFilters);
+			// Fill OutFilteredTestNames array with filtered test names
+			AutomationController->GetFilteredTestNames(OutFilteredTestNames);
+		}
 	}
 
 	void FindWorkers(float DeltaTime)
