AutomationTest - Use FApp::GetInstanceId() instead of Messaging Address to identify Automation worker

#rnx

diff --git a/Engine/Source/Developer/AutomationController/Private/AutomationControllerManager.cpp b/Engine/Source/Developer/AutomationController/Private/AutomationControllerManager.cpp
--- a/Engine/Source/Developer/AutomationController/Private/AutomationControllerManager.cpp
+++ b/Engine/Source/Developer/AutomationController/Private/AutomationControllerManager.cpp
@@ -280,7 +280,10 @@ void FAutomationControllerManager::RequestTests()
 			UE_LOG(LogAutomationController, Log, TEXT("Requesting test list from %s"), *MessageAddress.ToString());
 
 			//issue tests on appropriate platforms
-			MessageEndpoint->Send(FMessageEndpoint::MakeMessage<FAutomationWorkerRequestTests>(bDeveloperDirectoryIncluded, RequestedTestFlags), MessageAddress);
+			SendMessage(
+				FMessageEndpoint::MakeMessage<FAutomationWorkerRequestTests>(bDeveloperDirectoryIncluded, RequestedTestFlags),
+				FAutomationWorkerRequestTests::StaticStruct(),
+				MessageAddress);
 		}
 	}
 }
@@ -349,8 +352,12 @@ void FAutomationControllerManager::RunTests(const bool bInIsLocalSession)
 			FMessageAddress MessageAddress = DeviceClusterManager.GetDeviceMessageAddress(ClusterIndex, DeviceIndex);
 			UE_LOG(LogAutomationController, Log, TEXT("Sending StartTestSession to %s"), *MessageAddress.ToString());
 
-			MessageEndpoint->Send(FMessageEndpoint::MakeMessage<FAutomationWorkerStartTestSession>(), MessageAddress);
-			StartedTestSessionAddressSet.Add(MessageAddress);
+			SendMessage(
+				FMessageEndpoint::MakeMessage<FAutomationWorkerStartTestSession>(),
+				FAutomationWorkerStartTestSession::StaticStruct(),
+				MessageAddress);
+
+			StartedTestSessionWorkerInstanceIdSet.Add(DeviceClusterManager.GetClusterGameInstanceId(ClusterIndex, DeviceIndex));
 
 			// Store devices info into the json report.
 			if (JsonTestPassResults.IsRequired)
@@ -396,7 +403,10 @@ void FAutomationControllerManager::StopTests()
 				FMessageAddress MessageAddress = DeviceClusterManager.GetDeviceMessageAddress(ClusterIndex, DeviceIndex);
 
 				UE_LOG(LogAutomationController, Log, TEXT("Sending StopTests to %s"), *MessageAddress.ToString());
-				MessageEndpoint->Send(FMessageEndpoint::MakeMessage<FAutomationWorkerStopTests>(), MessageAddress);
+				SendMessage(
+					FMessageEndpoint::MakeMessage<FAutomationWorkerStopTests>(), 
+					FAutomationWorkerStopTests::StaticStruct(),
+					MessageAddress);
 			}
 		}
 
@@ -409,12 +419,7 @@ void FAutomationControllerManager::StopTests()
 
 	TestRunningArray.Empty();
 
-	for (const auto& MessageAddress : StartedTestSessionAddressSet)
-	{
-		UE_LOG(LogAutomationController, Log, TEXT("Sending StopTestSession to %s"), *MessageAddress.ToString());
-		MessageEndpoint->Send(FMessageEndpoint::MakeMessage<FAutomationWorkerStopTestSession>(), MessageAddress);
-	}
-	StartedTestSessionAddressSet.Empty();
+	StopStartedTestSessions();
 
 	// Close play window
 #if WITH_EDITOR
@@ -447,12 +452,12 @@ void FAutomationControllerManager::Tick()
 	ProcessComparisonQueue();
 }
 
-void FAutomationControllerManager::ReportImageComparisonResult(const FMessageAddress Sender, const FAutomationWorkerImageComparisonResults& Result)
+void FAutomationControllerManager::ReportImageComparisonResult(const FAutomationWorkerImageComparisonResults& Result)
 {
 	// Find the game session instance info
 	int32 ClusterIndex;
 	int32 DeviceIndex;
-	verify(DeviceClusterManager.FindDevice(Sender, ClusterIndex, DeviceIndex));
+	verify(DeviceClusterManager.FindDevice(Result.InstanceId, ClusterIndex, DeviceIndex));
 
 	// Get the current test.
 	TSharedPtr<IAutomationReport> Report = DeviceClusterManager.GetTest(ClusterIndex, DeviceIndex);
@@ -501,6 +506,7 @@ void FAutomationControllerManager::ProcessComparisonQueue()
 
 			FImageComparisonResult Result = Entry->PendingComparison.Get();
 			FAutomationWorkerImageComparisonResults ResultMessage(
+				Entry->InstanceId,
 				Result.bSkipAttachingImages ? FGuid() : FGuid::NewGuid(),
 				Result.ScreenshotName,
 				Result.IsNew(),
@@ -522,13 +528,13 @@ void FAutomationControllerManager::ProcessComparisonQueue()
 					*Entry->Sender.ToString()
 					, Result.IsNew()
 					, Result.AreSimilar()
-				);
-				MessageEndpoint->Send(Message, Entry->Sender);
+					);
+				SendMessage(Message, Message->StaticStruct(), Entry->Sender);
 			}
 
 			if (!Result.bSkipAttachingImages)
 			{
-				ReportImageComparisonResult(Entry->Sender, ResultMessage);
+				ReportImageComparisonResult(ResultMessage);
 			}
 		}
 	}
@@ -569,12 +575,7 @@ void FAutomationControllerManager::ProcessAvailableTasks()
 					#endif
 
 					// Notify the workers about stopping the session.
-					for (const auto& MessageAddress : StartedTestSessionAddressSet)
-					{
-						UE_LOG(LogAutomationController, Log, TEXT("Sending StopTestSession to %s"), *MessageAddress.ToString());
-						MessageEndpoint->Send(FMessageEndpoint::MakeMessage<FAutomationWorkerStopTestSession>(), MessageAddress);
-					}
-					StartedTestSessionAddressSet.Empty();
+					StopStartedTestSessions();
 
 					//Notify the graphical layout we are done processing results.
 					TestsCompleteDelegate.Broadcast();
@@ -608,6 +609,7 @@ void FAutomationControllerManager::CollectTestResults(TSharedPtr<IAutomationRepo
 		TestResult.SetEvents(Results.GetEntries(), Results.GetWarningTotal(), Results.GetErrorTotal());
 		TestResult.SetArtifacts(Results.Artifacts);
 		TestResult.Duration = Results.Duration;
+		
 		if (TestResult.DeviceInstance.IsEmpty())
 		{
 			TestResult.DeviceInstance = { Results.GameInstance };
@@ -808,8 +810,9 @@ void FAutomationControllerManager::ExecuteNextTask( int32 ClusterIndex, OUT bool
 						// Send it to each device
 						for (int32 AddressIndex = 0; AddressIndex < DeviceAddresses.Num(); ++AddressIndex)
 						{
+							const FAutomationDeviceInfo& DeviceInfo = DeviceClusterManager.GetDeviceInfo(ClusterIndex, DeviceIndex);
 							FAutomationTestResults TestResults;
-							TestResults.GameInstance = DeviceClusterManager.GetClusterGameInstance(ClusterIndex, DeviceIndex);
+							TestResults.GameInstance = DeviceInfo.Instance.ToString();
 							TestResults.State = EAutomationState::InProcess;
 							GameInstances.Add(TestResults.GameInstance);
 							NextTest->SetResults(ClusterIndex, CurrentTestPass, TestResults);
@@ -820,10 +823,13 @@ void FAutomationControllerManager::ExecuteNextTask( int32 ClusterIndex, OUT bool
 							// Send the test to the device for execution!
 							UE_LOG(LogAutomationController, Log, TEXT("Sending RunTest %s to %s"), *NextTest->GetDisplayName(), *DeviceAddress.ToString());
 
-							MessageEndpoint->Send(FMessageEndpoint::MakeMessage<FAutomationWorkerRunTests>(ExecutionCount, AddressIndex, NextTest->GetCommand(), NextTest->GetDisplayName(), NextTest->GetFullTestPath(), bSendAnalytics), DeviceAddress);
+							SendMessage(
+								FMessageEndpoint::MakeMessage<FAutomationWorkerRunTests>(ExecutionCount, AddressIndex, NextTest->GetCommand(), NextTest->GetDisplayName(), NextTest->GetFullTestPath(), bSendAnalytics),
+								FAutomationWorkerRunTests::StaticStruct(),
+								DeviceAddress);
 
 							// Add a test so we can check later if the device is still active
-							TestRunningArray.Add(FTestRunningInfo(DeviceAddress));
+							TestRunningArray.Add(FTestRunningInfo(DeviceAddress, DeviceInfo.Instance));
 						}
 
 						UE_LOG(LogAutomationController, Display, AutomationTestStarting, *NextTest->GetDisplayName(), *NextTest->GetFullTestPath());
@@ -925,13 +931,13 @@ void FAutomationControllerManager::RemoveCallbacks()
 	TestsCompleteDelegate.Clear();
 }
 
-void FAutomationControllerManager::SetTestNames(const FMessageAddress& AutomationWorkerAddress, TArray<FAutomationTestInfo>& TestInfo)
+void FAutomationControllerManager::SetTestNames(const FGuid& AutomationWorkerInstanceId, TArray<FAutomationTestInfo>& TestInfo)
 {
 	int32 DeviceClusterIndex = INDEX_NONE;
 	int32 DeviceIndex = INDEX_NONE;
 
 	// Find the device that requested these tests
-	if ( DeviceClusterManager.FindDevice(AutomationWorkerAddress, DeviceClusterIndex, DeviceIndex) )
+	if ( DeviceClusterManager.FindDevice(AutomationWorkerInstanceId, DeviceClusterIndex, DeviceIndex) )
 	{
 		// Sort tests by display name
 		struct FCompareAutomationTestInfo
@@ -1107,11 +1113,11 @@ void FAutomationControllerManager::SetControllerStatus(EAutomationControllerModu
 	}
 }
 
-void FAutomationControllerManager::RemoveTestRunning(const FMessageAddress& TestAddressToRemove)
+void FAutomationControllerManager::RemoveTestRunning(const FGuid& OwnerInstanceId)
 {
 	for ( int32 Index = 0; Index < TestRunningArray.Num(); Index++ )
 	{
-		if ( TestRunningArray[Index].OwnerMessageAddress == TestAddressToRemove )
+		if ( TestRunningArray[Index].OwnerInstanceId == OwnerInstanceId )
 		{
 			TestRunningArray.RemoveAt(Index);
 			break;
@@ -1119,11 +1125,11 @@ void FAutomationControllerManager::RemoveTestRunning(const FMessageAddress& Test
 	}
 }
 
-void FAutomationControllerManager::AddPingResult(const FMessageAddress& ResponderAddress)
+void FAutomationControllerManager::AddPingResult(const FGuid& ResponderInstanceId)
 {
 	for ( int32 Index = 0; Index < TestRunningArray.Num(); Index++ )
 	{
-		if ( TestRunningArray[Index].OwnerMessageAddress == ResponderAddress )
+		if (TestRunningArray[Index].OwnerInstanceId == ResponderInstanceId)
 		{
 			TestRunningArray[Index].LastPingTime = 0;
 			break;
@@ -1154,7 +1160,7 @@ void FAutomationControllerManager::UpdateTests()
 					// Find the game session instance info
 					int32 ClusterIndex;
 					int32 DeviceIndex;
-					verify(DeviceClusterManager.FindDevice(TestRunningArray[Index].OwnerMessageAddress, ClusterIndex, DeviceIndex));
+					verify(DeviceClusterManager.FindDevice(TestRunningArray[Index].OwnerInstanceId, ClusterIndex, DeviceIndex));
 					//verify this device thought it was busy
 					TSharedPtr <IAutomationReport> Report = DeviceClusterManager.GetTest(ClusterIndex, DeviceIndex);
 					check(Report.IsValid());
@@ -1163,7 +1169,7 @@ void FAutomationControllerManager::UpdateTests()
 
 					FAutomationTestResults TestResults;
 					TestResults.State = EAutomationState::Fail;
-					TestResults.GameInstance = DeviceClusterManager.GetClusterGameInstance(ClusterIndex, DeviceIndex);
+					TestResults.GameInstance = DeviceClusterManager.GetClusterGameInstanceId(ClusterIndex, DeviceIndex).ToString();
 					FString DeviceName = DeviceClusterManager.GetClusterDeviceName(ClusterIndex, DeviceIndex);
 					TestResults.AddEvent(FAutomationEvent(EAutomationEventType::Error, FString::Printf(TEXT("Timeout waiting for device %s"), *DeviceName)));
 
@@ -1210,7 +1216,10 @@ void FAutomationControllerManager::UpdateTests()
 				}
 				else
 				{
-					MessageEndpoint->Send(FMessageEndpoint::MakeMessage<FAutomationWorkerPing>(), TestRunningArray[Index].OwnerMessageAddress);
+					SendMessage(
+						FMessageEndpoint::MakeMessage<FAutomationWorkerPing>(),
+						FAutomationWorkerPing::StaticStruct(),
+						TestRunningArray[Index].OwnerMessageAddress);
 				}
 			}
 			CheckTestTimer = 0.f;
@@ -1222,6 +1231,41 @@ void FAutomationControllerManager::UpdateTests()
 	}
 }
 
+void FAutomationControllerManager::StopStartedTestSessions()
+{
+	int32 ClusterIndex = INDEX_NONE;
+	int32 DeviceIndex = INDEX_NONE;
+
+	// Notify the workers about stopping the session.
+	for (const auto& WorkerInstanceId : StartedTestSessionWorkerInstanceIdSet)
+	{
+		if (DeviceClusterManager.FindDevice(WorkerInstanceId, ClusterIndex, DeviceIndex))
+		{
+			const FMessageAddress MessageAddress = DeviceClusterManager.GetDeviceMessageAddress(ClusterIndex, DeviceIndex);
+			UE_LOG(LogAutomationController, Log, TEXT("Sending StopTestSession to %s"), *MessageAddress.ToString());
+			SendMessage(
+				FMessageEndpoint::MakeMessage<FAutomationWorkerStopTestSession>(),
+				FAutomationWorkerStopTestSession::StaticStruct(),
+				MessageAddress);
+		}
+	}
+}
+
+void FAutomationControllerManager::SendMessage(FAutomationWorkerMessageBase* Message, UScriptStruct* TypeInfo, const FMessageAddress& ControllerAddress)
+{
+	check(nullptr != Message);
+
+	Message->InstanceId = FApp::GetInstanceId();
+	MessageEndpoint->Send(
+		Message,
+		TypeInfo,
+		EMessageFlags::None,
+		nullptr,
+		TArrayBuilder<FMessageAddress>().Add(ControllerAddress),
+		FTimespan::Zero(),
+		FDateTime::MaxValue());
+}
+
 const bool FAutomationControllerManager::ExportReport(uint32 FileExportTypeMask)
 {
 	return ReportManager.ExportReport(FileExportTypeMask, GetNumDeviceClusters());
@@ -1271,12 +1315,14 @@ void FAutomationControllerManager::HandleFindWorkersResponseMessage(const FAutom
 
 void FAutomationControllerManager::HandlePongMessage( const FAutomationWorkerPong& Message, const TSharedRef<IMessageContext, ESPMode::ThreadSafe>& Context )
 {
-	AddPingResult(Context->GetSender());
+	DeviceClusterManager.UpdateDeviceFromMessage(Context->GetSender(), Message);
+	AddPingResult(Message.InstanceId);
 }
 
 void FAutomationControllerManager::HandleReceivedScreenShot(const FAutomationWorkerScreenImage& Message, const TSharedRef<IMessageContext, ESPMode::ThreadSafe>& Context)
 {
 	UE_LOG(LogAutomationController, Log, TEXT("ReceivedScreenShot for %s (%dx%d) from %s"), *Message.Metadata.TestName, Message.Metadata.Width, Message.Metadata.Height, *Context->GetSender().ToString());
+	DeviceClusterManager.UpdateDeviceFromMessage(Context->GetSender(), Message);
 
 	bool bTree = true;
 
@@ -1327,6 +1373,7 @@ void FAutomationControllerManager::HandleReceivedScreenShot(const FAutomationWor
 	// compare the incoming image and throw it away afterwards (note - there will be a copy in the report)
 	TSharedRef<FComparisonEntry> Comparison = MakeShareable(new FComparisonEntry());
 	Comparison->Sender = Context->GetSender();
+	Comparison->InstanceId = Message.InstanceId;
 	Comparison->PendingComparison = ScreenshotManager->CompareScreenshotAsync(IncomingFileName, Message.Metadata, EScreenShotCompareOptions::DiscardImage);
 
 	ComparisonQueue.Enqueue(Comparison);
@@ -1334,12 +1381,14 @@ void FAutomationControllerManager::HandleReceivedScreenShot(const FAutomationWor
 
 void FAutomationControllerManager::HandleReceivedComparisonResult(const FAutomationWorkerImageComparisonResults& Message, const TSharedRef<IMessageContext, ESPMode::ThreadSafe>& Context)
 {
-	ReportImageComparisonResult(Context->GetSender(), Message);
+	DeviceClusterManager.UpdateDeviceFromMessage(Context->GetSender(), Message);
+	ReportImageComparisonResult(Message);
 }
 
 void FAutomationControllerManager::HandleTestDataRequest(const FAutomationWorkerTestDataRequest& Message, const TSharedRef<IMessageContext, ESPMode::ThreadSafe>& Context)
 {
 	UE_LOG(LogAutomationController, Log, TEXT("Received TestDataRequest from %s"), *Context->GetSender().ToString());
+	DeviceClusterManager.UpdateDeviceFromMessage(Context->GetSender(), Message);
 
 	const FString TestDataRoot = FPaths::ConvertRelativePathToFull(FPaths::ProjectDir() / TEXT("Test"));
 	const FString DataFile = Message.DataType / Message.DataPlatform / Message.DataTestName / Message.DataName + TEXT(".json");
@@ -1383,7 +1432,7 @@ void FAutomationControllerManager::HandleTestDataRequest(const FAutomationWorker
 	ResponseMessage->bIsNew = bIsNew;
 	ResponseMessage->JsonData = ResponseJsonData;
 
-	MessageEndpoint->Send(ResponseMessage, Context->GetSender());
+	SendMessage(ResponseMessage, ResponseMessage->StaticStruct(), Context->GetSender());
 }
 
 void FAutomationControllerManager::HandlePerformanceDataRequest(const FAutomationWorkerPerformanceDataRequest& Message, const TSharedRef<IMessageContext, ESPMode::ThreadSafe>& Context)
@@ -1395,12 +1444,13 @@ void FAutomationControllerManager::HandlePerformanceDataRequest(const FAutomatio
 	ResponseMessage->bSuccess = true;
 	ResponseMessage->ErrorMessage = TEXT("");
 
-	MessageEndpoint->Send(ResponseMessage, Context->GetSender());
+	SendMessage(ResponseMessage, ResponseMessage->StaticStruct(), Context->GetSender());
 }
 
 void FAutomationControllerManager::HandleRequestNextNetworkCommandMessage(const FAutomationWorkerRequestNextNetworkCommand& Message, const TSharedRef<IMessageContext, ESPMode::ThreadSafe>& Context)
 {
 	UE_LOG(LogAutomationController, Log, TEXT("Received RequestNextNetworkCommandMessage from %s"), *Context->GetSender().ToString());
+	DeviceClusterManager.UpdateDeviceFromMessage(Context->GetSender(), Message);
 
 	// Harvest iteration of running the tests this result came from (stops stale results from being committed to subsequent runs)
 	if ( Message.ExecutionCount == ExecutionCount )
@@ -1409,7 +1459,7 @@ void FAutomationControllerManager::HandleRequestNextNetworkCommandMessage(const
 		int32 ClusterIndex;
 		int32 DeviceIndex;
 
-		verify(DeviceClusterManager.FindDevice(Context->GetSender(), ClusterIndex, DeviceIndex));
+		verify(DeviceClusterManager.FindDevice(Message.InstanceId, ClusterIndex, DeviceIndex));
 
 		// Verify this device thought it was busy
 		TSharedPtr<IAutomationReport> Report = DeviceClusterManager.GetTest(ClusterIndex, DeviceIndex);
@@ -1432,7 +1482,10 @@ void FAutomationControllerManager::HandleRequestNextNetworkCommandMessage(const
 			for ( int32 AddressIndex = 0; AddressIndex < DeviceAddresses.Num(); ++AddressIndex )
 			{
 				//send "next command message" to worker
-				MessageEndpoint->Send(FMessageEndpoint::MakeMessage<FAutomationWorkerNextNetworkCommandReply>(), DeviceAddresses[AddressIndex]);
+				SendMessage(
+					FMessageEndpoint::MakeMessage<FAutomationWorkerNextNetworkCommandReply>(),
+					FAutomationWorkerNextNetworkCommandReply::StaticStruct(),
+					DeviceAddresses[AddressIndex]	);
 			}
 		}
 	}
@@ -1441,6 +1494,7 @@ void FAutomationControllerManager::HandleRequestNextNetworkCommandMessage(const
 void FAutomationControllerManager::HandleRequestTestsReplyCompleteMessage(const FAutomationWorkerRequestTestsReplyComplete& Message, const TSharedRef<IMessageContext, ESPMode::ThreadSafe>& Context)
 {
 	UE_LOG(LogAutomationController, Log, TEXT("Received RequestTestsReplyCompleteMessage from %s"), *Context->GetSender().ToString());
+	DeviceClusterManager.UpdateDeviceFromMessage(Context->GetSender(), Message);
 
 	LLM_SCOPE_BYNAME(TEXT("AutomationTest/Controller"));
 	TArray<FAutomationTestInfo> TestInfo;
@@ -1453,11 +1507,12 @@ void FAutomationControllerManager::HandleRequestTestsReplyCompleteMessage(const
 
 	UE_LOG(LogAutomationController, Log, TEXT("%d tests available on %s"), TestInfo.Num(), *Context->GetSender().ToString());
 
-	SetTestNames(Context->GetSender(), TestInfo);
+	SetTestNames(Message.InstanceId, TestInfo);
 }
 
 void FAutomationControllerManager::HandleReceivedTelemetryData(const FAutomationWorkerTelemetryData& Message, const TSharedRef<IMessageContext, ESPMode::ThreadSafe>& Context)
 {
+	DeviceClusterManager.UpdateDeviceFromMessage(Context->GetSender(), Message);
 	FAutomationTelemetry::HandleAddTelemetry(Message);
 }
 
@@ -1549,6 +1604,8 @@ void FAutomationControllerManager::ReportAutomationResult(const TSharedPtr<IAuto
 
 void FAutomationControllerManager::HandleRunTestsReplyMessage(const FAutomationWorkerRunTestsReply& Message, const TSharedRef<IMessageContext, ESPMode::ThreadSafe>& Context)
 {
+	DeviceClusterManager.UpdateDeviceFromMessage(Context->GetSender(), Message);
+
 	// If we should commit these results
 	if ( Message.ExecutionCount == ExecutionCount )
 	{
@@ -1561,9 +1618,8 @@ void FAutomationControllerManager::HandleRunTestsReplyMessage(const FAutomationW
 		int32 ClusterIndex;
 		int32 DeviceIndex;
 
-		verify(DeviceClusterManager.FindDevice(Context->GetSender(), ClusterIndex, DeviceIndex));
-
-		TestResults.GameInstance = DeviceClusterManager.GetClusterGameInstance(ClusterIndex, DeviceIndex);
+		verify(DeviceClusterManager.FindDevice(Message.InstanceId, ClusterIndex, DeviceIndex));
+		TestResults.GameInstance = DeviceClusterManager.GetClusterGameInstanceId(ClusterIndex, DeviceIndex).ToString();
 		TestResults.SetEvents(Message.Entries, Message.WarningTotal, Message.ErrorTotal);
 
 		// Verify this device thought it was busy
@@ -1585,13 +1641,13 @@ void FAutomationControllerManager::HandleRunTestsReplyMessage(const FAutomationW
 	}
 
 	// Remove the running test
-	RemoveTestRunning(Context->GetSender());
+	RemoveTestRunning(Message.InstanceId);
 }
 
 void FAutomationControllerManager::HandleWorkerOfflineMessage( const FAutomationWorkerWorkerOffline& Message, const TSharedRef<IMessageContext, ESPMode::ThreadSafe>& Context )
 {
-	FMessageAddress DeviceMessageAddress = Context->GetSender();
-	DeviceClusterManager.Remove(DeviceMessageAddress);
+	FGuid DeviceInstanceId = Message.InstanceId;
+	DeviceClusterManager.Remove(DeviceInstanceId);
 }
 
 bool FAutomationControllerManager::IsDeviceGroupFlagSet( EAutomationDeviceGroupTypes::Type InDeviceGroup ) const
diff --git a/Engine/Source/Developer/AutomationController/Private/AutomationControllerManager.h b/Engine/Source/Developer/AutomationController/Private/AutomationControllerManager.h
--- a/Engine/Source/Developer/AutomationController/Private/AutomationControllerManager.h
+++ b/Engine/Source/Developer/AutomationController/Private/AutomationControllerManager.h
@@ -19,6 +19,8 @@
 #include "HAL/PlatformProperties.h"
 #include "AutomationControllerManager.generated.h"
 
+struct FAutomationWorkerMessageBase;
+
 USTRUCT()
 struct FAutomatedTestResult
 {
@@ -285,6 +287,11 @@ public:
 		return DeviceClusterManager.GetClusterDeviceName(ClusterIndex, DeviceIndex);
 	}
 
+	virtual FGuid GetGameInstanceId(const int32 ClusterIndex, const int32 DeviceIndex) const override
+	{
+		return DeviceClusterManager.GetClusterGameInstanceId(ClusterIndex, DeviceIndex);
+	}
+
 	virtual FString GetGameInstanceName(const int32 ClusterIndex, const int32 DeviceIndex) const override
 	{
 		return DeviceClusterManager.GetClusterGameInstance(ClusterIndex, DeviceIndex);
@@ -406,9 +413,9 @@ protected:
 	/**
 	 * Adds a ping result from a running test.
 	 *
-	 * @param ResponderAddress The address of the message endpoint that responded to a ping.
+	 * @param ResponderInstanceId The worker instance identifier that responded to a ping.
 	 */
-	void AddPingResult( const FMessageAddress& ResponderAddress );
+	void AddPingResult(const FGuid& ResponderInstanceId);
 
 	/**
 	* Spew all of our results of the test out to the log.
@@ -458,7 +465,7 @@ protected:
 	void ExecuteNextTask( int32 ClusterIndex, OUT bool& bAllTestsCompleted );
 
 	/* Report an image comparison result */
-	void ReportImageComparisonResult(const FMessageAddress Sender, const FAutomationWorkerImageComparisonResults& Result);
+	void ReportImageComparisonResult(const FAutomationWorkerImageComparisonResults& Result);
 
 	/** Process the comparison queue to see if there are comparisons we need to respond to the test with. */
 	void ProcessComparisonQueue();
@@ -472,19 +479,31 @@ protected:
 	/**
 	 * Removes the test info.
 	 *
-	 * @param TestToRemove The test to remove.
+	 * @param OwnerInstanceId Instance identifier of the test to remove.
 	 */
-	void RemoveTestRunning( const FMessageAddress& TestToRemove );
+	void RemoveTestRunning(const FGuid& OwnerInstanceId);
 
 	/** Changes the controller state. */
 	void SetControllerStatus( EAutomationControllerModuleState::Type AutomationTestState );
 
 	/** Stores the tests that are valid for a particular device classification. */
-	void SetTestNames(const FMessageAddress& AutomationWorkerAddress, TArray<FAutomationTestInfo>& TestInfo);
+	void SetTestNames(const FGuid& AutomationWorkerInstanceId, TArray<FAutomationTestInfo>& TestInfo);
 
 	/** Updates the tests to ensure they are all still running. */
 	void UpdateTests();
 
+	/** Sends stop session message to worker instances. */
+	void StopStartedTestSessions();
+
+	/**
+	 * Send a message in an unified way with passing correct Instance Id into the message we are going to send.
+	 *
+	 * @param Message The message to be sent.
+	 * @param TypeInfo The type information about the message to be sent.
+	 * @param ControllerAddress The message address of the receiver.
+	 */
+	void SendMessage(FAutomationWorkerMessageBase* Message, UScriptStruct* TypeInfo, const FMessageAddress& ControllerAddress);
+
 private:
 
 	/** Handles FAutomationWorkerFindWorkersResponse messages. */
@@ -587,13 +606,16 @@ private:
 	/** A data holder to keep track of how long tests have been running. */
 	struct FTestRunningInfo
 	{
-		FTestRunningInfo( FMessageAddress InMessageAddress ):
+		FTestRunningInfo( FMessageAddress InMessageAddress, FGuid InInstanceId):
 			OwnerMessageAddress( InMessageAddress),
+			OwnerInstanceId( InInstanceId ),
 			LastPingTime( 0.f )
 		{
 		}
 		/** The test runners message address */
 		FMessageAddress OwnerMessageAddress;
+		/** The test runner's instance ID */
+		FGuid OwnerInstanceId;
 		/** The time since we had a ping from the instance*/
 		double LastPingTime;
 	};
@@ -601,8 +623,8 @@ private:
 	/** A array of running tests. */
 	TArray< FTestRunningInfo > TestRunningArray;
 
-	/** Set of addresses of started sessions. */
-	TSet< FMessageAddress > StartedTestSessionAddressSet;
+	/** Set of worker instance identifiers for started sessions. */
+	TSet< FGuid > StartedTestSessionWorkerInstanceIdSet;
 
 	/** The number of test passes to perform. */
 	int32 NumTestPasses = 0;
@@ -628,6 +650,7 @@ private:
 	struct FComparisonEntry
 	{
 		FMessageAddress Sender;
+		FGuid InstanceId;
 		TFuture<FImageComparisonResult> PendingComparison;
 	};
 
diff --git a/Engine/Source/Developer/AutomationController/Private/AutomationDeviceClusterManager.cpp b/Engine/Source/Developer/AutomationController/Private/AutomationDeviceClusterManager.cpp
--- a/Engine/Source/Developer/AutomationController/Private/AutomationDeviceClusterManager.cpp
+++ b/Engine/Source/Developer/AutomationController/Private/AutomationDeviceClusterManager.cpp
@@ -10,13 +10,14 @@ void FAutomationDeviceClusterManager::Reset()
 	Clusters.Empty();
 }
 
-
 void FAutomationDeviceClusterManager::AddDeviceFromMessage(const FMessageAddress& MessageAddress, const FAutomationWorkerFindWorkersResponse& Message, const uint32 GroupFlags)
 {
-	int32 TestClusterIndex;
-	int32 TestDeviceIndex;
+	check(Message.InstanceId.IsValid());
+
+	int32 TestClusterIndex = INDEX_NONE;
+	int32 TestDeviceIndex = INDEX_NONE;
 	//if we don't already know about this device
-	if (!FindDevice(MessageAddress, TestClusterIndex, TestDeviceIndex))
+	if (!FindDevice(Message.InstanceId, TestClusterIndex, TestDeviceIndex))
 	{
 		FDeviceState NewDevice(MessageAddress, Message);
 		FString GroupName = GetGroupNameForDevice(NewDevice, GroupFlags);
@@ -41,16 +42,33 @@ void FAutomationDeviceClusterManager::AddDeviceFromMessage(const FMessageAddress
 			Clusters.Add(NewCluster);
 		}
 	}
+	else
+	{
+		UpdateDeviceFromMessage(MessageAddress, Message);
+	}
 }
 
+void FAutomationDeviceClusterManager::UpdateDeviceFromMessage(const FMessageAddress& MessageAddress, const FAutomationWorkerMessageBase& Message)
+{
+	check(Message.InstanceId.IsValid());
+
+	int32 TestClusterIndex = INDEX_NONE;
+	int32 TestDeviceIndex = INDEX_NONE;
 
-void FAutomationDeviceClusterManager::Remove(const FMessageAddress& MessageAddress)
+	if (FindDevice(Message.InstanceId, TestClusterIndex, TestDeviceIndex))
+	{
+		// If we already know about this device
+		Clusters[TestClusterIndex].Devices[TestDeviceIndex].DeviceMessageAddress = MessageAddress;
+	}
+}
+
+void FAutomationDeviceClusterManager::Remove(const FGuid& DeviceInstanceId)
 {
 	for (int32 ClusterIndex = 0; ClusterIndex < Clusters.Num(); ++ClusterIndex)
 	{
 		for (int32 DeviceIndex = Clusters[ClusterIndex].Devices.Num()-1; DeviceIndex >= 0; --DeviceIndex)
 		{
-			if (MessageAddress == Clusters[ClusterIndex].Devices[DeviceIndex].DeviceMessageAddress)
+			if (DeviceInstanceId == Clusters[ClusterIndex].Devices[DeviceIndex].Info.Instance)
 			{
 				Clusters[ClusterIndex].Devices.RemoveAt(DeviceIndex);
 			}
@@ -212,10 +230,16 @@ FString FAutomationDeviceClusterManager::GetClusterDeviceName(const int32 Cluste
 }
 
 FString FAutomationDeviceClusterManager::GetClusterGameInstance(const int32 ClusterIndex, const int32 DeviceIndex) const
+{
+	return GetDeviceInfo(ClusterIndex, DeviceIndex).InstanceName;
+}
+
+FGuid FAutomationDeviceClusterManager::GetClusterGameInstanceId(const int32 ClusterIndex, const int32 DeviceIndex) const
 {
 	return GetDeviceInfo(ClusterIndex, DeviceIndex).Instance;
 }
 
+
 const FAutomationDeviceInfo& FAutomationDeviceClusterManager::GetDeviceInfo(const int32 ClusterIndex, const int32 DeviceIndex) const
 {
 	check((ClusterIndex >= 0) && (ClusterIndex < Clusters.Num()));
@@ -223,7 +247,7 @@ const FAutomationDeviceInfo& FAutomationDeviceClusterManager::GetDeviceInfo(cons
 	return Clusters[ClusterIndex].Devices[DeviceIndex].Info;
 }
 
-bool FAutomationDeviceClusterManager::FindDevice(const FMessageAddress& MessageAddress, int32& OutClusterIndex, int32& OutDeviceIndex)
+bool FAutomationDeviceClusterManager::FindDevice(const FGuid& InstanceId, int32& OutClusterIndex, int32& OutDeviceIndex)
 {
 	OutClusterIndex = INDEX_NONE;
 	OutDeviceIndex = INDEX_NONE;
@@ -232,7 +256,7 @@ bool FAutomationDeviceClusterManager::FindDevice(const FMessageAddress& MessageA
 		for (int32 DeviceIndex = 0; DeviceIndex < Clusters[ClusterIndex].Devices.Num(); ++DeviceIndex)
 		{
 			//if network addresses match
-			if (MessageAddress == Clusters[ClusterIndex].Devices[DeviceIndex].DeviceMessageAddress)
+			if (InstanceId == Clusters[ClusterIndex].Devices[DeviceIndex].Info.Instance)
 			{
 				OutClusterIndex = ClusterIndex;
 				OutDeviceIndex = DeviceIndex;
@@ -243,7 +267,6 @@ bool FAutomationDeviceClusterManager::FindDevice(const FMessageAddress& MessageA
 	return false;
 }
 
-
 FMessageAddress FAutomationDeviceClusterManager::GetDeviceMessageAddress(const int32 ClusterIndex, const int32 DeviceIndex) const
 {
 	//verify cluster/device index
diff --git a/Engine/Source/Developer/AutomationController/Private/AutomationDeviceClusterManager.h b/Engine/Source/Developer/AutomationController/Private/AutomationDeviceClusterManager.h
--- a/Engine/Source/Developer/AutomationController/Private/AutomationDeviceClusterManager.h
+++ b/Engine/Source/Developer/AutomationController/Private/AutomationDeviceClusterManager.h
@@ -21,9 +21,13 @@ public:
 	UPROPERTY()
 	FString DeviceName;
 
-	/** The instance name */
+	/** The instance ID */
 	UPROPERTY()
-	FString Instance;
+	FGuid Instance;
+
+	/** The instance Name */
+	UPROPERTY()
+	FString InstanceName;
 
 	/** The name of the platform */
 	UPROPERTY()
@@ -80,8 +84,16 @@ public:
 	 */
 	void AddDeviceFromMessage(const FMessageAddress& MessageAddress, const FAutomationWorkerFindWorkersResponse& Message, const uint32 GroupFlags);
 
+	/**
+	 * Update existing device's network address.
+	 *
+	 * @param MessageAddress The network address of an available worker.
+	 * @param Message The message that contains the device info.
+	 */
+	void UpdateDeviceFromMessage(const FMessageAddress& MessageAddress, const FAutomationWorkerMessageBase& Message);
+
 	/** Remove a device (went offline, etc). */
-	void Remove(const FMessageAddress& MessageAddress);
+	void Remove(const FGuid& DeviceInstanceId);
 
 	/** Returns number of unique device types. */
 	int32 GetNumClusters() const;
@@ -107,18 +119,21 @@ public:
 	/** Returns the instance of the game within a cluster. */
 	FString GetClusterGameInstance(const int32 ClusterIndex, const int32 DeviceIndex) const;
 
+	/** Returns the instance identifier of the game within a cluster. */
+	FGuid GetClusterGameInstanceId(const int32 ClusterIndex, const int32 DeviceIndex) const;
+
 	/** Returns the name of a device within a cluster. */
 	const FAutomationDeviceInfo& GetDeviceInfo(const int32 ClusterIndex, const int32 DeviceIndex) const;
 
-	/** 
+	/**
 	 * Finds the cluster/device index for a particular GUID.
 	 *
-	 * @param MessageAddress Network address of the device.
+	 * @param InstanceId The device instance identifier.
 	 * @param OutClusterIndex The index of the platform this device was in.
 	 * @param OutDeviceIndex The index of the device within OutClusterIndex.
 	 * @return true if the device was found within a cluster.
 	 */
-	bool FindDevice(const FMessageAddress& MessageAddress, int32& OutClusterIndex, int32& OutDeviceIndex);
+	bool FindDevice(const FGuid& InstanceId, int32& OutClusterIndex, int32& OutDeviceIndex);
 
 	/** 
 	 * Returns the message address of the device specified.
@@ -209,7 +224,8 @@ private:
 			Info.RAMInGB = Message.RAMInGB;
 			Info.RenderMode = Message.RenderModeName;
 			Info.RHI = Message.RHIName;
-			Info.Instance = Message.InstanceName;
+			Info.Instance = Message.InstanceId;
+			Info.InstanceName = Message.InstanceName;
 			Report.Reset();
 			IsDeviceAvailable = true;
 		}
diff --git a/Engine/Source/Developer/AutomationController/Public/IAutomationControllerManager.h b/Engine/Source/Developer/AutomationController/Public/IAutomationControllerManager.h
--- a/Engine/Source/Developer/AutomationController/Public/IAutomationControllerManager.h
+++ b/Engine/Source/Developer/AutomationController/Public/IAutomationControllerManager.h
@@ -274,6 +274,15 @@ public:
 	 */
 	virtual FString GetDeviceTypeName(const int32 ClusterIndex) const = 0;
 
+	/**
+	 * Get a game instance ID.
+	 *
+	 * @param ClusterIndex The cluster Index.
+	 * @param DeviceIndex The Device Index.
+	 * @return identifier of a game.
+	 */
+	virtual FGuid GetGameInstanceId(const int32 ClusterIndex, const int32 DeviceIndex) const = 0;
+
 	/**
 	 * Get a game instance name.
 	 *
diff --git a/Engine/Source/Developer/AutomationController/Public/IAutomationReport.h b/Engine/Source/Developer/AutomationController/Public/IAutomationReport.h
--- a/Engine/Source/Developer/AutomationController/Public/IAutomationReport.h
+++ b/Engine/Source/Developer/AutomationController/Public/IAutomationReport.h
@@ -139,7 +139,7 @@ public:
 	/* The time this test took to complete */
 	float Duration;
 
-	/* The name of the instance which reported these results */
+	/* The instance which reported these results */
 	FString GameInstance;
 
 	/** Artifacts generated during the run of the test. */
@@ -397,8 +397,8 @@ public:
 	 * @param ClusterIndex - Index of the platform reporting the results of this test.  See AutomationDeviceClusterManager.
 	 * @return the name of the device.
 	 */
-	virtual FString GetGameInstanceName( const int32 ClusterIndex ) = 0;
-
+	virtual FString GetGameInstanceName(const int32 ClusterIndex) = 0;
+	
 	/**
 	 * Add a child test to the hierarchy, creating internal tree nodes as needed.
 	 * If NewTestName is Editor.Maps.Loadall.TestName, this will create nodes for Editor, Maps, Loadall, and then a leaf node for the test name with the associated command line
diff --git a/Engine/Source/Developer/AutomationWindow/Private/SAutomationGraphicalResultBox.cpp b/Engine/Source/Developer/AutomationWindow/Private/SAutomationGraphicalResultBox.cpp
--- a/Engine/Source/Developer/AutomationWindow/Private/SAutomationGraphicalResultBox.cpp
+++ b/Engine/Source/Developer/AutomationWindow/Private/SAutomationGraphicalResultBox.cpp
@@ -91,12 +91,12 @@ void SAutomationGraphicalResultBox::PopulateData()
 		ClusterIt->TotalNumTests = EnabledReports.Num();
 		for(int j = 0; j < AutomationController->GetNumDevicesInCluster(i); ++j)
 		{
-			ClusterIt->Devices.Add(FDeviceResults(AutomationController->GetGameInstanceName(i, j)));
+			ClusterIt->Devices.Add(FDeviceResults(AutomationController->GetGameInstanceId(i, j), AutomationController->GetGameInstanceName(i, j)));
 			FDeviceResults* DeviceIt = &ClusterIt->Devices[j];
 			for( int k = 0; k < EnabledReports.Num(); ++k )
 			{
 				FAutomationTestResults TestResults = EnabledReports[k]->GetResults( i, LastPassIndex );
-				if( TestResults.GameInstance == DeviceIt->InstanceName )
+				if (TestResults.GameInstance == DeviceIt->GameInstanceId)
 				{
 					uint32 TestIndex = DeviceIt->Tests.Add(FTestResults(EnabledReports[k]->GetDisplayNameWithDecoration()));
 					FTestResults* TestIt = &DeviceIt->Tests[TestIndex];
@@ -199,7 +199,7 @@ void SAutomationGraphicalResultBox::CreateWidgets()
 					+SVerticalBox::Slot()
 					[
 						SNew(STextBlock)
-						.Text(FText::FromString(DeviceIt->InstanceName))
+						.Text(FText::FromString(DeviceIt->GameInstanceName))
 					]
 					+SVerticalBox::Slot()
 					[
diff --git a/Engine/Source/Developer/AutomationWindow/Private/SAutomationGraphicalResultBox.h b/Engine/Source/Developer/AutomationWindow/Private/SAutomationGraphicalResultBox.h
--- a/Engine/Source/Developer/AutomationWindow/Private/SAutomationGraphicalResultBox.h
+++ b/Engine/Source/Developer/AutomationWindow/Private/SAutomationGraphicalResultBox.h
@@ -98,16 +98,20 @@ private:
 	{
 	public:
 		/** Constructor */
-		FDeviceResults(const FString& InName)
+		FDeviceResults(const FGuid& InGameInstanceId, const FString& InGameInstanceName)
 		{
-			InstanceName = InName;
+			GameInstanceName = InGameInstanceName;
+			GameInstanceId = InGameInstanceId.ToString();
 			TotalTime = 0.f;
 			TotalTestSkips = 0;
 			TotalTestSuccesses = 0;
 		}
 
-		/** Name of the device instance */
-		FString InstanceName;
+		/** The game instance Name running on the Device */
+		FString GameInstanceName;
+
+		/** The game instance Id running on the Device */
+		FString GameInstanceId;
 
 		/** Total time of all the tests */
 		float TotalTime;
diff --git a/Engine/Source/Developer/AutomationWindow/Private/SAutomationTestItem.cpp b/Engine/Source/Developer/AutomationWindow/Private/SAutomationTestItem.cpp
--- a/Engine/Source/Developer/AutomationWindow/Private/SAutomationTestItem.cpp
+++ b/Engine/Source/Developer/AutomationWindow/Private/SAutomationTestItem.cpp
@@ -368,19 +368,19 @@ FText SAutomationTestItem::GetTestToolTip( int32 ClusterIndex ) const
 	else
 	{
 		FFormatNamedArguments Args;
-		Args.Add(TEXT("GameName"), FText::FromString(TestStatus->GetGameInstanceName(ClusterIndex)));
+		Args.Add(TEXT("GameInstance"), FText::FromString(TestStatus->GetGameInstanceName(ClusterIndex)));
 
 		if (TestState == EAutomationState::InProcess)
 		{
-			TestToolTip = FText::Format(LOCTEXT("TestToolTipInProgress", "In progress on: {GameName}"), Args);
+			TestToolTip = FText::Format(LOCTEXT("TestToolTipInProgress", "In progress on: {GameInstance}"), Args);
 		}
 		else if (TestState == EAutomationState::Success)
 		{
-			TestToolTip = FText::Format(LOCTEXT("TestToolTipComplete", "Completed on: {GameName}"), Args);
+			TestToolTip = FText::Format(LOCTEXT("TestToolTipComplete", "Completed on: {GameInstance}"), Args);
 		}
 		else
 		{
-			TestToolTip = FText::Format(LOCTEXT("TestToolTipFailed", "Failed on: {GameName}"), Args);
+			TestToolTip = FText::Format(LOCTEXT("TestToolTipFailed", "Failed on: {GameInstance}"), Args);
 		}
 	}
 	return TestToolTip;
@@ -407,11 +407,13 @@ FSlateColor SAutomationTestItem::IsToBeSkipped_GetColorAndOpacity() const
 	{
 		return FLinearColor(1.0f, 1.0f, 1.0f, 0.4f);
 	}
+
 	// Identify visually if the test is to be skipped on certain condition.
 	if (TestStatus->IsToBeSkippedOnConditions())
 	{
 		return FLinearColor(1.0f, 1.0f, 0.2f, 0.6f);
 	}
+
 	return FLinearColor(1.0f, 1.0f, 1.0f, 1.0f);
 }
 
diff --git a/Engine/Source/Runtime/AutomationMessages/Public/AutomationWorkerMessages.h b/Engine/Source/Runtime/AutomationMessages/Public/AutomationWorkerMessages.h
--- a/Engine/Source/Runtime/AutomationMessages/Public/AutomationWorkerMessages.h
+++ b/Engine/Source/Runtime/AutomationMessages/Public/AutomationWorkerMessages.h
@@ -11,6 +11,16 @@
 
 #include "AutomationWorkerMessages.generated.h"
 
+USTRUCT()
+struct FAutomationWorkerMessageBase
+{
+	GENERATED_USTRUCT_BODY()
+
+public:
+
+	UPROPERTY(EditAnywhere, Category = "Message")
+	FGuid InstanceId = FGuid{};
+};
 
 /* Worker discovery messages
  *****************************************************************************/
@@ -19,7 +29,7 @@
  * Implements a message that is published to find automation workers.
  */
 USTRUCT()
-struct FAutomationWorkerFindWorkers
+struct FAutomationWorkerFindWorkers : public FAutomationWorkerMessageBase
 {
 	GENERATED_USTRUCT_BODY()
 
@@ -56,7 +66,7 @@ struct FAutomationWorkerFindWorkers
  * Implements a message that is sent in response to FAutomationWorkerFindWorkers.
  */
 USTRUCT()
-struct FAutomationWorkerFindWorkersResponse
+struct FAutomationWorkerFindWorkersResponse : public FAutomationWorkerMessageBase
 {
 	GENERATED_USTRUCT_BODY()
 
@@ -113,7 +123,7 @@ struct FAutomationWorkerFindWorkersResponse
  * Implements a message that notifies automation controllers that a worker went off-line.
  */
 USTRUCT()
-struct FAutomationWorkerWorkerOffline
+struct FAutomationWorkerWorkerOffline : public FAutomationWorkerMessageBase
 {
 	GENERATED_USTRUCT_BODY()
 };
@@ -122,7 +132,7 @@ struct FAutomationWorkerWorkerOffline
 /**
  */
 USTRUCT()
-struct FAutomationWorkerPing
+struct FAutomationWorkerPing : public FAutomationWorkerMessageBase
 {
 	GENERATED_USTRUCT_BODY()
 };
@@ -131,7 +141,7 @@ struct FAutomationWorkerPing
 /**
  */
 USTRUCT()
-struct FAutomationWorkerStartTestSession
+struct FAutomationWorkerStartTestSession : public FAutomationWorkerMessageBase
 {
 	GENERATED_USTRUCT_BODY()
 };
@@ -140,7 +150,7 @@ struct FAutomationWorkerStartTestSession
 /**
  */
 USTRUCT()
-struct FAutomationWorkerStopTestSession
+struct FAutomationWorkerStopTestSession : public FAutomationWorkerMessageBase
 {
 	GENERATED_USTRUCT_BODY()
 };
@@ -149,7 +159,7 @@ struct FAutomationWorkerStopTestSession
 /**
 */
 USTRUCT()
-struct FAutomationWorkerStopTests
+struct FAutomationWorkerStopTests : public FAutomationWorkerMessageBase
 {
 	GENERATED_USTRUCT_BODY()
 };
@@ -158,7 +168,7 @@ struct FAutomationWorkerStopTests
 /**
  */
 USTRUCT()
-struct FAutomationWorkerPong
+struct FAutomationWorkerPong : public FAutomationWorkerMessageBase
 {
 	GENERATED_USTRUCT_BODY()
 };
@@ -168,7 +178,7 @@ struct FAutomationWorkerPong
  * Implements a message for requesting available automation tests from a worker.
  */
 USTRUCT()
-struct FAutomationWorkerRequestTests
+struct FAutomationWorkerRequestTests : public FAutomationWorkerMessageBase
 {
 	GENERATED_USTRUCT_BODY()
 
@@ -195,7 +205,7 @@ struct FAutomationWorkerRequestTests
  * A single test reply, used by FAutomationWorkerRequestTestsReplyComplete
  */
 USTRUCT()
-struct FAutomationWorkerSingleTestReply
+struct FAutomationWorkerSingleTestReply : public FAutomationWorkerMessageBase
 {
 	GENERATED_USTRUCT_BODY()
 
@@ -268,7 +278,7 @@ struct FAutomationWorkerSingleTestReply
  * Returns list of all tests
  */
 USTRUCT()
-struct FAutomationWorkerRequestTestsReplyComplete
+struct FAutomationWorkerRequestTestsReplyComplete : public FAutomationWorkerMessageBase
 {
 	GENERATED_USTRUCT_BODY()
 
@@ -281,7 +291,7 @@ struct FAutomationWorkerRequestTestsReplyComplete
  * Implements a message to request the running of automation tests on a worker.
  */
 USTRUCT()
-struct FAutomationWorkerRunTests
+struct FAutomationWorkerRunTests : public FAutomationWorkerMessageBase
 {
 	GENERATED_USTRUCT_BODY()
 
@@ -328,12 +338,11 @@ struct FAutomationWorkerRunTests
  * Implements a message that is sent in response to FAutomationWorkerRunTests.
  */
 USTRUCT()
-struct FAutomationWorkerRunTestsReply
+struct FAutomationWorkerRunTestsReply : public FAutomationWorkerMessageBase
 {
 	GENERATED_USTRUCT_BODY()
 
 public:
-
 	/** */
 	UPROPERTY(EditAnywhere, Category="Message")
 	FString TestName;
@@ -365,7 +374,7 @@ public:
 /**
  */
 USTRUCT()
-struct FAutomationWorkerRequestNextNetworkCommand
+struct FAutomationWorkerRequestNextNetworkCommand : public FAutomationWorkerMessageBase
 {
 	GENERATED_USTRUCT_BODY()
 
@@ -386,7 +395,7 @@ struct FAutomationWorkerRequestNextNetworkCommand
 /**
  */
 USTRUCT()
-struct FAutomationWorkerNextNetworkCommandReply
+struct FAutomationWorkerNextNetworkCommandReply : public FAutomationWorkerMessageBase
 {
 	GENERATED_USTRUCT_BODY()
 };
@@ -700,7 +709,7 @@ public:
  * Implements a message that is sent in containing a screen shot run during performance test.
  */
 USTRUCT()
-struct FAutomationWorkerScreenImage
+struct FAutomationWorkerScreenImage : public FAutomationWorkerMessageBase
 {
 	GENERATED_USTRUCT_BODY()
 
@@ -726,14 +735,15 @@ struct FAutomationWorkerScreenImage
  * Implements a message that is sent in containing a screen shot run during performance test.
  */
 USTRUCT()
-struct FAutomationWorkerImageComparisonResults
+struct FAutomationWorkerImageComparisonResults : public FAutomationWorkerMessageBase
 {
 	GENERATED_USTRUCT_BODY()
 
 public:
 
-	FAutomationWorkerImageComparisonResults()
-		: bNew(false)
+	FAutomationWorkerImageComparisonResults(FGuid InInstanceId = FGuid{})
+		: FAutomationWorkerMessageBase{InInstanceId}
+		, bNew(false)
 		, bSimilar(false)
 		, MaxLocalDifference(0.0)
 		, GlobalDifference(0.0)
@@ -741,6 +751,7 @@ public:
 	}
 
 	FAutomationWorkerImageComparisonResults(
+		FGuid InInstanceId,
 		FGuid InUniqueId,
 		const FString& InName,
 		bool InIsNew,
@@ -753,7 +764,8 @@ public:
 		const FString& InReportApprovedFilePath,
 		const FString& InReportIncomingFilePath
 	)
-		: UniqueId(InUniqueId)
+		: FAutomationWorkerMessageBase{ InInstanceId }
+		, UniqueId(InUniqueId)
 		, ScreenshotName(InName)
 		, bNew(InIsNew)
 		, bSimilar(InAreSimilar)
@@ -811,7 +823,7 @@ public:
  * for the first time this test is run, it might need to store things, or get things.
  */
 USTRUCT()
-struct FAutomationWorkerTestDataRequest
+struct FAutomationWorkerTestDataRequest : public FAutomationWorkerMessageBase
 {
 	GENERATED_USTRUCT_BODY()
 
@@ -840,7 +852,7 @@ struct FAutomationWorkerTestDataRequest
  * Implements a message that responds to TestDataRequests.
  */
 USTRUCT()
-struct FAutomationWorkerTestDataResponse
+struct FAutomationWorkerTestDataResponse : public FAutomationWorkerMessageBase
 {
 	GENERATED_USTRUCT_BODY()
 
@@ -856,7 +868,7 @@ struct FAutomationWorkerTestDataResponse
  * Implements a message to request the performance data for this hardware.
  */
 USTRUCT()
-struct FAutomationWorkerPerformanceDataRequest
+struct FAutomationWorkerPerformanceDataRequest : public FAutomationWorkerMessageBase
 {
 	GENERATED_USTRUCT_BODY()
 
@@ -881,7 +893,7 @@ struct FAutomationWorkerPerformanceDataRequest
  * Implements a message that responds to PerformanceDataRequest.
  */
 USTRUCT()
-struct FAutomationWorkerPerformanceDataResponse
+struct FAutomationWorkerPerformanceDataResponse : public FAutomationWorkerMessageBase
 {
 	GENERATED_USTRUCT_BODY()
 
@@ -898,7 +910,7 @@ struct FAutomationWorkerPerformanceDataResponse
  * Implements a message that contains telemetry data point.
  */
 USTRUCT()
-struct FAutomationWorkerTelemetryItem
+struct FAutomationWorkerTelemetryItem : public FAutomationWorkerMessageBase
 {
 	GENERATED_USTRUCT_BODY()
 
@@ -935,7 +947,7 @@ struct FAutomationWorkerTelemetryItem
  * Implements a message that contains telemetry data.
  */
 USTRUCT()
-struct FAutomationWorkerTelemetryData
+struct FAutomationWorkerTelemetryData : public FAutomationWorkerMessageBase
 {
 	GENERATED_USTRUCT_BODY()
 
diff --git a/Engine/Source/Runtime/AutomationWorker/Private/AutomationWorkerModule.cpp b/Engine/Source/Runtime/AutomationWorker/Private/AutomationWorkerModule.cpp
--- a/Engine/Source/Runtime/AutomationWorker/Private/AutomationWorkerModule.cpp
+++ b/Engine/Source/Runtime/AutomationWorker/Private/AutomationWorkerModule.cpp
@@ -36,7 +36,6 @@ DEFINE_LOG_CATEGORY_STATIC(LogAutomationWorker, Log, All);
 
 IMPLEMENT_MODULE(FAutomationWorkerModule, AutomationWorker);
 
-
 /* IModuleInterface interface
  *****************************************************************************/
 
@@ -172,7 +171,10 @@ void FAutomationWorkerModule::ReportNetworkCommandComplete()
 {
 	if (GIsAutomationTesting)
 	{
-		MessageEndpoint->Send(FMessageEndpoint::MakeMessage<FAutomationWorkerRequestNextNetworkCommand>(ExecutionCount), TestRequesterAddress);
+		SendMessage(
+			FMessageEndpoint::MakeMessage<FAutomationWorkerRequestNextNetworkCommand>(ExecutionCount),
+			FAutomationWorkerRequestNextNetworkCommand::StaticStruct(),
+			TestRequesterAddress);
 		if (StopTestEvent.IsBound())
 		{
 			// this is a local test; the message to continue will never arrive, so lets not wait for it
@@ -226,7 +228,7 @@ void FAutomationWorkerModule::ReportTestComplete()
 				HandleTelemetryData(ExecutionInfo.TelemetryStorage, FullTestPath, ExecutionInfo.TelemetryItems);
 			}
 
-			MessageEndpoint->Send(Message, TestRequesterAddress);
+			SendMessage(Message, Message->StaticStruct(), TestRequesterAddress);
 		}
 
 
@@ -252,9 +254,23 @@ void FAutomationWorkerModule::SendTests( const FMessageAddress& ControllerAddres
 
 	UE_LOG(LogAutomationWorker, Log, TEXT("Set %d tests to %s"), TestInfo.Num(), *ControllerAddress.ToString());
 
-	MessageEndpoint->Send(Reply, ControllerAddress);
+	SendMessage(Reply, Reply->StaticStruct(), ControllerAddress);
 }
 
+void FAutomationWorkerModule::SendMessage(FAutomationWorkerMessageBase* Message, UScriptStruct* TypeInfo, const FMessageAddress& ControllerAddress)
+{
+	check(nullptr != Message);
+
+	Message->InstanceId = FApp::GetInstanceId();
+	MessageEndpoint->Send(
+		Message,
+		TypeInfo,
+		EMessageFlags::None,
+		nullptr,
+		TArrayBuilder<FMessageAddress>().Add(ControllerAddress),
+		FTimespan::Zero(),
+		FDateTime::MaxValue());
+}
 
 /* FAutomationWorkerModule callbacks
  *****************************************************************************/
@@ -301,10 +317,9 @@ void FAutomationWorkerModule::SendWorkerFound(const FMessageAddress& ControllerA
 	FString CPUModelString = FPlatformMisc::GetCPUBrand().TrimStart();
 
 	FString DeviceName = DeviceTag.IsEmpty() ? FPlatformProcess::ComputerName() : DeviceTag;
-	FString DeviceId = FPlatformMisc::GetDeviceId().IsEmpty() ? DeviceName : FPlatformMisc::GetDeviceId();
 
 	Response->DeviceName = DeviceName;
-	Response->InstanceName = FString::Printf(TEXT("%s-%s-%s"), *DeviceId, *FApp::GetSessionId().ToString(), *FApp::GetInstanceId().ToString());
+	Response->InstanceName = FApp::GetInstanceName();
 	Response->Platform = FPlatformProperties::PlatformName();
 	Response->SessionId = FApp::GetSessionId();
 	Response->OSVersionName = OSVersionString;
@@ -319,7 +334,7 @@ void FAutomationWorkerModule::SendWorkerFound(const FMessageAddress& ControllerA
 	Response->RenderModeName = TEXT("Unknown");
 #endif
 
-	MessageEndpoint->Send(Response, ControllerAddress);
+	SendMessage(Response, Response->StaticStruct(), ControllerAddress);
 }
 
 
@@ -337,7 +352,7 @@ void FAutomationWorkerModule::HandleNextNetworkCommandReplyMessage( const FAutom
 
 void FAutomationWorkerModule::HandlePingMessage( const FAutomationWorkerPing& Message, const TSharedRef<IMessageContext, ESPMode::ThreadSafe>& Context )
 {
-	MessageEndpoint->Send(FMessageEndpoint::MakeMessage<FAutomationWorkerPong>(), Context->GetSender());
+	SendMessage(FMessageEndpoint::MakeMessage<FAutomationWorkerPong>(), FAutomationWorkerPong::StaticStruct(), Context->GetSender());
 }
 
 
@@ -442,7 +457,7 @@ void FAutomationWorkerModule::HandleScreenShotComparisonReport(const FAutomation
 	Message->ReportApprovedFilePath = Results.ReportApprovedFilePath;
 	Message->ReportIncomingFilePath = Results.ReportIncomingFilePath;
 
-	MessageEndpoint->Send(Message, TestRequesterAddress);
+	SendMessage(Message, Message->StaticStruct(), TestRequesterAddress);
 }
 
 void FAutomationWorkerModule::HandleScreenShotCapturedWithName(const TArray<FColor>& RawImageData, const FAutomationScreenshotData& Data)
@@ -474,7 +489,7 @@ void FAutomationWorkerModule::HandleScreenShotAndTraceCapturedWithName(const TAr
 
 		UE_LOG(LogAutomationWorker, Log, TEXT("Sending screenshot %s to %s"), *Message->ScreenShotName, *TestRequesterAddress.ToString());
 
-		MessageEndpoint->Send(Message, TestRequesterAddress);
+		SendMessage(Message, Message->StaticStruct(), TestRequesterAddress);
 	}
 	else
 	{
@@ -619,7 +634,7 @@ void FAutomationWorkerModule::HandleRunTestsMessage( const FAutomationWorkerRunT
 		OutMessage->State = EAutomationState::Skipped;
 		OutMessage->Entries.Add(FAutomationExecutionEntry(FAutomationEvent(EAutomationEventType::Error, LogMessage)));
 		OutMessage->ErrorTotal = 1;
-		MessageEndpoint->Send(OutMessage, Context->GetSender());
+		SendMessage(OutMessage, OutMessage->StaticStruct(), Context->GetSender());
 
 		return;
 	}
@@ -648,7 +663,7 @@ void FAutomationWorkerModule::HandleRunTestsMessage( const FAutomationWorkerRunT
 		OutMessage->ExecutionCount = Message.ExecutionCount;
 		OutMessage->State = EAutomationState::Skipped;
 		OutMessage->Entries.Add(FAutomationExecutionEntry(FAutomationEvent(EAutomationEventType::Info, FString::Printf(TEXT("Skipping test: %s"), *SkipReason))));
-		MessageEndpoint->Send(OutMessage, Context->GetSender());
+		SendMessage(OutMessage, OutMessage->StaticStruct(), Context->GetSender());
 
 		return;
 	}
@@ -792,7 +807,7 @@ void FAutomationWorkerModule::HandleTelemetryData(const FString& StorageName, co
 
 	UE_LOG(LogAutomationWorker, Log, TEXT("Sending Telemetry Data for %s"), *Message->TestName);
 
-	MessageEndpoint->Send(Message, TestRequesterAddress);
+	SendMessage(Message, Message->StaticStruct(), TestRequesterAddress);
 }
 
 void FAutomationWorkerModule::RecordPerformanceAnalytics( const FAutomationPerformanceSnapshot& PerfSnapshot )
diff --git a/Engine/Source/Runtime/AutomationWorker/Private/AutomationWorkerModule.h b/Engine/Source/Runtime/AutomationWorker/Private/AutomationWorkerModule.h
--- a/Engine/Source/Runtime/AutomationWorker/Private/AutomationWorkerModule.h
+++ b/Engine/Source/Runtime/AutomationWorker/Private/AutomationWorkerModule.h
@@ -23,6 +23,7 @@ struct FAutomationWorkerStartTestSession;
 struct FAutomationWorkerStopTestSession;
 struct FAutomationWorkerRunTests;
 struct FMessageAddress;
+struct FAutomationWorkerMessageBase;
 
 
 /**
@@ -69,6 +70,15 @@ protected:
 	 */
 	void SendTests(const FMessageAddress& ControllerAddress);
 
+	/**
+	 * Send a message in an unified way with passing correct Instance Id into the message we are going to send.
+	 * 
+	 * @param Message The message to be sent.
+	 * @param TypeInfo The type information about the message to be sent.
+	 * @param ControllerAddress The message address of the receiver.
+	 */
+	void SendMessage(FAutomationWorkerMessageBase* Message, UScriptStruct* TypeInfo, const FMessageAddress& ControllerAddress);
+
 private:
 
 	/** Handles FAutomationWorkerFindWorkers messages. */
