Fix exclusion not working in Test Automation search field
Avoid always adding child items if parents are included.

#rnx

diff --git a/Engine/Source/Developer/AutomationController/Private/AutomationReport.cpp b/Engine/Source/Developer/AutomationController/Private/AutomationReport.cpp
--- a/Engine/Source/Developer/AutomationController/Private/AutomationReport.cpp
+++ b/Engine/Source/Developer/AutomationController/Private/AutomationReport.cpp
@@ -318,7 +318,7 @@ bool FAutomationReport::SetFilter( TSharedPtr< AutomationFilterCollection > InFi
 	{
 		bool ThisChildPassedFilter = ChildReports[ChildIndex]->SetFilter( InFilter, bSelfPassesFilter );
 
-		if( ThisChildPassedFilter || bSelfPassesFilter || ParentPassedFilter )
+		if( ThisChildPassedFilter )
 		{
 			if ( !ChildReports[ChildIndex]->IsParent() || ChildReports[ChildIndex]->GetFilteredChildren().Num() > 0 )
 			{
