Fix issue with Automation commands conflicting with Automation CVar

#rnx

diff --git a/Engine/Source/Developer/AutomationController/Private/AutomationCommandline.cpp b/Engine/Source/Developer/AutomationController/Private/AutomationCommandline.cpp
--- a/Engine/Source/Developer/AutomationController/Private/AutomationCommandline.cpp
+++ b/Engine/Source/Developer/AutomationController/Private/AutomationCommandline.cpp
@@ -515,18 +515,22 @@ protected:
 	virtual bool Exec_Dev(UWorld*, const TCHAR* Cmd, FOutputDevice& Ar) override
 	{
 		bool bHandled = false;
-		// Track whether we have a flag we care about passing through.
-		FString FlagToUse = "";
-
 		// figure out if we are handling this request
 		if (FParse::Command(&Cmd, TEXT("Automation")))
 		{
+			// Early exit in case of a CVar input. ie: Automation.SkipStackWalk 1
+			if (FString(Cmd).StartsWith(TEXT(".")))
+			{
+				return false;
+			}
+
+			// Track whether we have a flag we care about passing through.
+			FString FlagToUse = "";
+			uint32 TestFlag = 0;
 
 			TArray<FString> CommandList;
 			FString(Cmd).ParseIntoArray(CommandList, TEXT(";"), true);
 
-			Init();
-
 			//assume we handle this
 			bHandled = true;
 
@@ -576,23 +580,23 @@ protected:
 					Ar.Logf(TEXT("Automation: Setting minimum priority of cases to run to: %s"), *FlagToUse);
 					if (FlagToUse.Contains(TEXT("Low")))
 					{
-						AutomationController->SetRequestedTestFlags(EAutomationTestFlags::PriorityMask);
+						TestFlag = EAutomationTestFlags::PriorityMask;
 					}
 					else if (FlagToUse.Contains(TEXT("Medium")))
 					{
-						AutomationController->SetRequestedTestFlags(EAutomationTestFlags::MediumPriorityAndAbove);
+						TestFlag = EAutomationTestFlags::MediumPriorityAndAbove;
 					}
 					else if (FlagToUse.Contains(TEXT("High")))
 					{
-						AutomationController->SetRequestedTestFlags(EAutomationTestFlags::HighPriorityAndAbove);
+						TestFlag = EAutomationTestFlags::HighPriorityAndAbove;
 					}
 					else if (FlagToUse.Contains(TEXT("Critical")))
 					{
-						AutomationController->SetRequestedTestFlags(EAutomationTestFlags::CriticalPriority);
+						TestFlag = EAutomationTestFlags::CriticalPriority;
 					}
 					else if (FlagToUse.Contains(TEXT("None")))
 					{
-						AutomationController->SetRequestedTestFlags(0);
+						TestFlag = 0;
 					}
 					else
 					{
@@ -605,23 +609,23 @@ protected:
 					Ar.Logf(TEXT("Setting explicit priority of cases to run to: %s"), *FlagToUse);
 					if (FlagToUse.Contains(TEXT("Low")))
 					{
-						AutomationController->SetRequestedTestFlags(EAutomationTestFlags::LowPriority);
+						TestFlag = EAutomationTestFlags::LowPriority;
 					}
 					else if (FlagToUse.Contains(TEXT("Medium")))
 					{
-						AutomationController->SetRequestedTestFlags(EAutomationTestFlags::MediumPriority);
+						TestFlag = EAutomationTestFlags::MediumPriority;
 					}
 					else if (FlagToUse.Contains(TEXT("High")))
 					{
-						AutomationController->SetRequestedTestFlags(EAutomationTestFlags::HighPriority);
+						TestFlag = EAutomationTestFlags::HighPriority;
 					}
 					else if (FlagToUse.Contains(TEXT("Critical")))
 					{
-						AutomationController->SetRequestedTestFlags(EAutomationTestFlags::CriticalPriority);
+						TestFlag = EAutomationTestFlags::CriticalPriority;
 					}
 					else if (FlagToUse.Contains(TEXT("None")))
 					{
-						AutomationController->SetRequestedTestFlags(0);
+						TestFlag = 0;
 					}
 
 					else
@@ -641,7 +645,7 @@ protected:
 					StringCommand = TempCmd;
 					if (FilterMaps.Contains(FlagToUse))
 					{
-						AutomationController->SetRequestedTestFlags(FilterMaps[FlagToUse]);
+						TestFlag = FilterMaps[FlagToUse];
 						Ar.Logf(TEXT("Automation: RunFilter='%s' Queued."), *FlagToUse);
 					}
 					AutomationCommandQueue.Add(EAutomationCommand::RunFilter);
@@ -651,7 +655,7 @@ protected:
 					FlagToUse = TempCmd;
 					if (FilterMaps.Contains(FlagToUse))
 					{
-						AutomationController->SetRequestedTestFlags(FilterMaps[FlagToUse]);
+						TestFlag = FilterMaps[FlagToUse];
 						Ar.Logf(TEXT("Automation: Setting test filter: %s"), *FlagToUse);
 					}
 				}
@@ -702,9 +706,9 @@ protected:
 						CVar->Set(true);
 					}
 				}
-				else
+				else if (FParse::Command(&TempCmd, TEXT("Help")))
 				{
-					Ar.Logf(TEXT("Incorrect automation command syntax! Supported commands are: "));
+					Ar.Logf(TEXT("Supported commands are: "));
 					Ar.Logf(TEXT("\tAutomation StartRemoteSession <sessionid>"));
 					Ar.Logf(TEXT("\tAutomation List"));
 					Ar.Logf(TEXT("\tAutomation RunTests <test string>"));
@@ -720,6 +724,20 @@ protected:
 					Ar.Logf(TEXT("\tAutomation EnableStereoTests"));
 					bHandled = false;
 				}
+				else
+				{
+					Ar.Logf(TEXT("Unknown Automation command '%s'! Use Help command for a detailed list."), TempCmd);
+					bHandled = false;
+				}
+			}
+
+			if (bHandled)
+			{
+				Init();
+				if (TestFlag != 0)
+				{
+					AutomationController->SetRequestedTestFlags(TestFlag);
+				}
 			}
 		}
 		
