Fix an issue with Automation Command line not being initiated properly
#rnx

diff --git a/Engine/Source/Developer/AutomationController/Private/AutomationCommandline.cpp b/Engine/Source/Developer/AutomationController/Private/AutomationCommandline.cpp
--- a/Engine/Source/Developer/AutomationController/Private/AutomationCommandline.cpp
+++ b/Engine/Source/Developer/AutomationController/Private/AutomationCommandline.cpp
@@ -526,11 +526,12 @@ protected:
 
 			// Track whether we have a flag we care about passing through.
 			FString FlagToUse = "";
-			uint32 TestFlag = 0;
 
 			TArray<FString> CommandList;
 			FString(Cmd).ParseIntoArray(CommandList, TEXT(";"), true);
 
+			Init();
+
 			//assume we handle this
 			bHandled = true;
 
@@ -580,23 +581,23 @@ protected:
 					Ar.Logf(TEXT("Automation: Setting minimum priority of cases to run to: %s"), *FlagToUse);
 					if (FlagToUse.Contains(TEXT("Low")))
 					{
-						TestFlag = EAutomationTestFlags::PriorityMask;
+						AutomationController->SetRequestedTestFlags(EAutomationTestFlags::PriorityMask);
 					}
 					else if (FlagToUse.Contains(TEXT("Medium")))
 					{
-						TestFlag = EAutomationTestFlags::MediumPriorityAndAbove;
+						AutomationController->SetRequestedTestFlags(EAutomationTestFlags::MediumPriorityAndAbove);
 					}
 					else if (FlagToUse.Contains(TEXT("High")))
 					{
-						TestFlag = EAutomationTestFlags::HighPriorityAndAbove;
+						AutomationController->SetRequestedTestFlags(EAutomationTestFlags::HighPriorityAndAbove);
 					}
 					else if (FlagToUse.Contains(TEXT("Critical")))
 					{
-						TestFlag = EAutomationTestFlags::CriticalPriority;
+						AutomationController->SetRequestedTestFlags(EAutomationTestFlags::CriticalPriority);
 					}
 					else if (FlagToUse.Contains(TEXT("None")))
 					{
-						TestFlag = 0;
+						AutomationController->SetRequestedTestFlags(0);
 					}
 					else
 					{
@@ -609,23 +610,23 @@ protected:
 					Ar.Logf(TEXT("Setting explicit priority of cases to run to: %s"), *FlagToUse);
 					if (FlagToUse.Contains(TEXT("Low")))
 					{
-						TestFlag = EAutomationTestFlags::LowPriority;
+						AutomationController->SetRequestedTestFlags(EAutomationTestFlags::LowPriority);
 					}
 					else if (FlagToUse.Contains(TEXT("Medium")))
 					{
-						TestFlag = EAutomationTestFlags::MediumPriority;
+						AutomationController->SetRequestedTestFlags(EAutomationTestFlags::MediumPriority);
 					}
 					else if (FlagToUse.Contains(TEXT("High")))
 					{
-						TestFlag = EAutomationTestFlags::HighPriority;
+						AutomationController->SetRequestedTestFlags(EAutomationTestFlags::HighPriority);
 					}
 					else if (FlagToUse.Contains(TEXT("Critical")))
 					{
-						TestFlag = EAutomationTestFlags::CriticalPriority;
+						AutomationController->SetRequestedTestFlags(EAutomationTestFlags::CriticalPriority);
 					}
 					else if (FlagToUse.Contains(TEXT("None")))
 					{
-						TestFlag = 0;
+						AutomationController->SetRequestedTestFlags(0);
 					}
 
 					else
@@ -645,7 +646,7 @@ protected:
 					StringCommand = TempCmd;
 					if (FilterMaps.Contains(FlagToUse))
 					{
-						TestFlag = FilterMaps[FlagToUse];
+						AutomationController->SetRequestedTestFlags(FilterMaps[FlagToUse]);
 						Ar.Logf(TEXT("Automation: RunFilter='%s' Queued."), *FlagToUse);
 					}
 					AutomationCommandQueue.Add(EAutomationCommand::RunFilter);
@@ -655,7 +656,7 @@ protected:
 					FlagToUse = TempCmd;
 					if (FilterMaps.Contains(FlagToUse))
 					{
-						TestFlag = FilterMaps[FlagToUse];
+						AutomationController->SetRequestedTestFlags(FilterMaps[FlagToUse]);
 						Ar.Logf(TEXT("Automation: Setting test filter: %s"), *FlagToUse);
 					}
 				}
@@ -730,15 +731,6 @@ protected:
 					bHandled = false;
 				}
 			}
-
-			if (bHandled)
-			{
-				Init();
-				if (TestFlag != 0)
-				{
-					AutomationController->SetRequestedTestFlags(TestFlag);
-				}
-			}
 		}
 		
 		// Shutdown our service
