Support exclude condition in Automation Test Group filter
#rnx

diff --git a/Engine/Source/Developer/AutomationController/Public/AutomationControllerSettings.h b/Engine/Source/Developer/AutomationController/Public/AutomationControllerSettings.h
--- a/Engine/Source/Developer/AutomationController/Public/AutomationControllerSettings.h
+++ b/Engine/Source/Developer/AutomationController/Public/AutomationControllerSettings.h
@@ -7,6 +7,7 @@
 #include "Containers/UnrealString.h"
 #include "CoreMinimal.h"
 #include "HAL/Platform.h"
+#include "IAutomationReport.h"
 #include "UObject/Object.h"
 #include "UObject/ObjectMacros.h"
 #include "UObject/SoftObjectPath.h"
@@ -15,21 +16,23 @@
 #include "AutomationControllerSettings.generated.h"
 
 /*
-* Describes a filter for a test group.
+* Describes a base filter for a test group.
 */
 USTRUCT()
-struct FAutomatedTestFilter
+struct FAutomatedTestFilterBase
 {
-	GENERATED_USTRUCT_BODY()
+	GENERATED_BODY()
 
 public:
 
-	FAutomatedTestFilter(FString InContains, bool InMatchFromStart = false, bool InMatchFromEnd = false)
+	FAutomatedTestFilterBase(FString InContains, bool InMatchFromStart = false, bool InMatchFromEnd = false)
 		: Contains(InContains), MatchFromStart(InMatchFromStart), MatchFromEnd(InMatchFromEnd)
 	{
 	}
 
-	FAutomatedTestFilter() : FAutomatedTestFilter(TEXT("")) {}
+	FAutomatedTestFilterBase() : FAutomatedTestFilterBase(TEXT("")) {}
+
+	virtual ~FAutomatedTestFilterBase() {}
 
 	/** String that the test must contain */
 	UPROPERTY(Config)
@@ -41,7 +44,72 @@ public:
 
 	/** If true start matching from the end of the string, else anywhere */
 	UPROPERTY(Config)
-		bool MatchFromEnd;	
+		bool MatchFromEnd;
+
+	virtual bool PassesFilter(const TSharedPtr< IAutomationReport >& InReport) const
+	{
+		bool bMeetsMatch = true;
+
+		if (MatchFromStart || MatchFromEnd)
+		{
+			if (MatchFromStart)
+			{
+				bMeetsMatch = InReport->GetFullTestPath().StartsWith(Contains);
+			}
+
+			if (MatchFromEnd && bMeetsMatch)
+			{
+				bMeetsMatch = InReport->GetFullTestPath().EndsWith(Contains);
+			}
+		}
+		else
+		{
+			bMeetsMatch = InReport->GetFullTestPath().Contains(Contains);
+		}
+
+		return bMeetsMatch;
+	}
+};
+
+/*
+* Describes a filter for a test group with exclude option.
+*/
+USTRUCT()
+struct FAutomatedTestFilter : public FAutomatedTestFilterBase
+{
+	GENERATED_BODY()
+
+public:
+
+	FAutomatedTestFilter(FString InContains, bool InMatchFromStart = false, bool InMatchFromEnd = false)
+		: FAutomatedTestFilterBase(InContains, InMatchFromStart, InMatchFromEnd)
+	{
+	}
+
+	FAutomatedTestFilter() : FAutomatedTestFilter(TEXT("")) {}
+
+	/** List of filters to exclude */
+	UPROPERTY(Config)
+		TArray<FAutomatedTestFilterBase> Exclude;
+
+	virtual bool PassesFilter(const TSharedPtr< IAutomationReport >& InReport) const override
+	{
+		bool bMeetsMatch = Super::PassesFilter(InReport);
+
+		if (bMeetsMatch && !Exclude.IsEmpty())
+		{
+			// Apply exclusion rules
+			for (const FAutomatedTestFilterBase& Filter : Exclude)
+			{
+				if (Filter.PassesFilter(InReport))
+				{
+					return false;
+				}
+			}
+		}
+
+		return bMeetsMatch;
+	}
 };
 
 /*
diff --git a/Engine/Source/Developer/AutomationController/Public/AutomationGroupFilter.h b/Engine/Source/Developer/AutomationController/Public/AutomationGroupFilter.h
--- a/Engine/Source/Developer/AutomationController/Public/AutomationGroupFilter.h
+++ b/Engine/Source/Developer/AutomationController/Public/AutomationGroupFilter.h
@@ -48,35 +48,14 @@ public:
 
 	virtual bool PassesFilter( const TSharedPtr< IAutomationReport >& InReport ) const override
 	{
-		const int NumGroupFilters = Filters.Num();
-		bool bMeetsMatch;
-
 		for (const FAutomatedTestFilter& Filter: Filters)
 		{
-			bMeetsMatch = true;
-
-			if (Filter.MatchFromStart || Filter.MatchFromEnd)
-			{
-				if (Filter.MatchFromStart)
-				{
-					bMeetsMatch = InReport->GetFullTestPath().StartsWith(Filter.Contains);
-				}
-
-				if (Filter.MatchFromEnd && bMeetsMatch)
-				{
-					bMeetsMatch = InReport->GetFullTestPath().EndsWith(Filter.Contains);
-				}
-			}
-			else
-			{
-				bMeetsMatch = InReport->GetFullTestPath().Contains(Filter.Contains);
-			}
-
-			if (bMeetsMatch)
+			if (Filter.PassesFilter(InReport))
 			{
 				return true;
 			}
 		}
+
 		return Filters.Num() == 0;
 	}
 
