Use Utc time in test automation report

#rnx

diff --git a/Engine/Source/Developer/AutomationController/Private/AutomationControllerManager.cpp b/Engine/Source/Developer/AutomationController/Private/AutomationControllerManager.cpp
--- a/Engine/Source/Developer/AutomationController/Private/AutomationControllerManager.cpp
+++ b/Engine/Source/Developer/AutomationController/Private/AutomationControllerManager.cpp
@@ -137,7 +137,7 @@ void FAutomatedTestPassResults::UpdateTestResultStatus(const IAutomationReportPt
 		Failed++;
 		break;
 	case EAutomationState::InProcess:
-		TestResult.DateTime = FDateTime::Now();
+		TestResult.DateTime = FDateTime::UtcNow();
 		InProcess++;
 		break;
 	case EAutomationState::NotRun:
@@ -663,7 +663,7 @@ void FAutomationControllerManager::CollectTestResults(TSharedPtr<IAutomationRepo
 bool FAutomationControllerManager::GenerateJsonTestPassSummary(FAutomatedTestPassResults& SerializedPassResults)
 {
 	UE_LOG(LogAutomationController, Display, TEXT("Converting results to json object..."));
-	SerializedPassResults.ReportCreatedOn = FDateTime::Now();
+	SerializedPassResults.ReportCreatedOn = FDateTime::UtcNow();
 
 	FString Json;
 	if (FJsonObjectConverter::UStructToJsonObjectString(SerializedPassResults, Json))
