Revert FAutomationArtifact key mapping constraint as enum; use FString instead as key

#rnx

diff --git a/Engine/Source/Developer/AutomationController/Private/AutomationControllerManager.cpp b/Engine/Source/Developer/AutomationController/Private/AutomationControllerManager.cpp
--- a/Engine/Source/Developer/AutomationController/Private/AutomationControllerManager.cpp
+++ b/Engine/Source/Developer/AutomationController/Private/AutomationControllerManager.cpp
@@ -464,12 +464,12 @@ void FAutomationControllerManager::ReportImageComparisonResult(const FAutomation
 	if (Report.IsValid())
 	{
 		// Record the artifacts for the test.
-		TMap<EComparisonFileTypes, FString> LocalFiles;
+		TMap<FString, FString> LocalFiles;
 
 		FString ScreenshotResultsFolder = FPaths::AutomationReportsDir();
 
 		// Paths in the result are relative to the automation report directory.	
-		LocalFiles.Add(EComparisonFileTypes::Unapproved, FPaths::Combine(ScreenshotResultsFolder, Result.ReportIncomingFilePath));
+		LocalFiles.Add(ComparisonFileTypes::Unapproved, FPaths::Combine(ScreenshotResultsFolder, Result.ReportIncomingFilePath));
 
 		// Don't copy reference and delta if the images are similar.
 		if (!Result.bSimilar)
@@ -477,12 +477,12 @@ void FAutomationControllerManager::ReportImageComparisonResult(const FAutomation
 			// unapproved should always be valid. but approved/difference may be empty if this is a new screenshot
 			if (Result.ReportIncomingFilePath.Len())
 			{
-				LocalFiles.Add(EComparisonFileTypes::Approved, FPaths::Combine(ScreenshotResultsFolder, Result.ReportApprovedFilePath));
+				LocalFiles.Add(ComparisonFileTypes::Approved, FPaths::Combine(ScreenshotResultsFolder, Result.ReportApprovedFilePath));
 			}
 
 			if (Result.ReportComparisonFilePath.Len())
 			{
-				LocalFiles.Add(EComparisonFileTypes::Difference, FPaths::Combine(ScreenshotResultsFolder, Result.ReportComparisonFilePath));
+				LocalFiles.Add(ComparisonFileTypes::Difference, FPaths::Combine(ScreenshotResultsFolder, Result.ReportComparisonFilePath));
 			}
 		}
 
@@ -628,14 +628,14 @@ void FAutomationControllerManager::CollectTestResults(TSharedPtr<IAutomationRepo
 		FCriticalSection CS;
 		for (FAutomationArtifact& Artifact : TestResult.GetArtifacts())
 		{
-			TArray<EComparisonFileTypes> Keys;
+			TArray<FString> Keys;
 			Artifact.LocalFiles.GetKeys(Keys);
 
-			bool bOnlyUnapproved = Keys.Num() == 1 && Keys[0] == EComparisonFileTypes::Unapproved;
+			bool bOnlyUnapproved = Keys.Num() == 1 && Keys[0] == ComparisonFileTypes::Unapproved;
 
 			ParallelFor(Keys.Num(), [&](int32 Index)
 				{
-					const EComparisonFileTypes& Key = Keys[Index];
+					const FString& Key = Keys[Index];
 					FString Path = Artifact.LocalFiles[Key];
 					FPaths::MakePathRelativeTo(Path, *FPaths::AutomationReportsDir());
 					Path = ArtifactDirName / Path;
@@ -643,7 +643,7 @@ void FAutomationControllerManager::CollectTestResults(TSharedPtr<IAutomationRepo
 						FScopeLock Lock(&CS);
 						Artifact.Files.Add(Key, MoveTemp(Path));
 					}
-					if (Key == EComparisonFileTypes::Unapproved)
+					if (Key == ComparisonFileTypes::Unapproved && Artifact.Type == EAutomationArtifactType::Comparison)
 					{
 						// Copy screenshot report
 						FScreenshotExportResult ExportResult = ScreenshotManager->ExportScreenshotComparisonResult(Artifact.Name, ArtifactExportPath, bOnlyUnapproved);
diff --git a/Engine/Source/Developer/AutomationController/Public/IAutomationReport.h b/Engine/Source/Developer/AutomationController/Public/IAutomationReport.h
--- a/Engine/Source/Developer/AutomationController/Public/IAutomationReport.h
+++ b/Engine/Source/Developer/AutomationController/Public/IAutomationReport.h
@@ -31,13 +31,14 @@ enum class EAutomationArtifactType : uint8
 	Comparison
 };
 
-UENUM()
-enum class EComparisonFileTypes : uint8
+namespace ComparisonFileTypes
 {
-	Unapproved,
-	Approved,
-	Difference
-};
+	// Use const string here instead of an enum, because FAutomationArtifact.Files property is a dictionary string/string
+	// FAutomationArtifact should be able to be used as an adhoc key/file association.
+	const FString Unapproved = TEXT("unapproved");
+	const FString Approved = TEXT("approved");
+	const FString Difference = TEXT("difference");
+}
 
 USTRUCT()
 struct FAutomationArtifact
@@ -51,7 +52,7 @@ public:
 	{
 	}
 
-	FAutomationArtifact(FGuid InUniqueId, const FString& InName, EAutomationArtifactType InType, const TMap<EComparisonFileTypes, FString>& InLocalFiles)
+	FAutomationArtifact(FGuid InUniqueId, const FString& InName, EAutomationArtifactType InType, const TMap<FString, FString>& InLocalFiles)
 		: Id(InUniqueId)
 		, Name(InName)
 		, Type(InType)
@@ -71,12 +72,12 @@ public:
 	EAutomationArtifactType Type;
 
 	UPROPERTY()
-	TMap<EComparisonFileTypes, FString> Files;
+	TMap<FString, FString> Files;
 
 	// Local Files are the files generated during a testing run, once exported, the individual file paths
 	// should be stored in the Files map.
 
-	TMap<EComparisonFileTypes, FString> LocalFiles;
+	TMap<FString, FString> LocalFiles;
 };
 
 /**
