Align commandline filter options to test automation window filter options.
Add Standard and Negative option

#rnx

diff --git a/Engine/Source/Developer/AutomationController/Private/AutomationCommandline.cpp b/Engine/Source/Developer/AutomationController/Private/AutomationCommandline.cpp
--- a/Engine/Source/Developer/AutomationController/Private/AutomationCommandline.cpp
+++ b/Engine/Source/Developer/AutomationController/Private/AutomationCommandline.cpp
@@ -96,6 +96,8 @@ public:
 		FilterMaps.Add("Stress", EAutomationTestFlags::StressFilter);
 		FilterMaps.Add("Perf", EAutomationTestFlags::PerfFilter);
 		FilterMaps.Add("Product", EAutomationTestFlags::ProductFilter);
+		FilterMaps.Add("Standard", EAutomationTestFlags::SmokeFilter | EAutomationTestFlags::EngineFilter | EAutomationTestFlags::ProductFilter | EAutomationTestFlags::PerfFilter);
+		FilterMaps.Add("Negative", EAutomationTestFlags::NegativeFilter);
 		FilterMaps.Add("All", EAutomationTestFlags::FilterMask);
 	}
 	
@@ -669,6 +671,7 @@ protected:
 						continue;
 					}
 					AutomationCommandQueue.Add(EAutomationCommand::RunAll);
+					AutomationController->SetRequestedTestFlags(FilterMaps["All"]);
 					Ar.Logf(TEXT("Automation: RunAll Queued. NOTE: This may take a while."));
 				}
 				else if (FParse::Command(&TempCmd, TEXT("Quit")))
