Copy Delta and report file when copying image comparison report

#rnx

diff --git a/Engine/Source/Developer/AutomationController/Private/AutomationControllerManager.cpp b/Engine/Source/Developer/AutomationController/Private/AutomationControllerManager.cpp
--- a/Engine/Source/Developer/AutomationController/Private/AutomationControllerManager.cpp
+++ b/Engine/Source/Developer/AutomationController/Private/AutomationControllerManager.cpp
@@ -471,19 +471,19 @@ void FAutomationControllerManager::ReportImageComparisonResult(const FAutomation
 		// Paths in the result are relative to the automation report directory.	
 		LocalFiles.Add(ComparisonFileTypes::Unapproved, FPaths::Combine(ScreenshotResultsFolder, Result.ReportIncomingFilePath));
 
-		// Don't copy reference and delta if the images are similar.
+		// unapproved should always be valid; but approved/difference may be empty if this is a new screenshot
+		if (Result.ReportComparisonFilePath.Len())
+		{
+			LocalFiles.Add(ComparisonFileTypes::Difference, FPaths::Combine(ScreenshotResultsFolder, Result.ReportComparisonFilePath));
+		}
+
+		// Don't copy reference if the images are similar.
 		if (!Result.bSimilar)
 		{
-			// unapproved should always be valid. but approved/difference may be empty if this is a new screenshot
 			if (Result.ReportIncomingFilePath.Len())
 			{
 				LocalFiles.Add(ComparisonFileTypes::Approved, FPaths::Combine(ScreenshotResultsFolder, Result.ReportApprovedFilePath));
 			}
-
-			if (Result.ReportComparisonFilePath.Len())
-			{
-				LocalFiles.Add(ComparisonFileTypes::Difference, FPaths::Combine(ScreenshotResultsFolder, Result.ReportComparisonFilePath));
-			}
 		}
 
 		Report->AddArtifact(ClusterIndex, CurrentTestPass, FAutomationArtifact(Result.UniqueId, Result.ScreenshotPath, EAutomationArtifactType::Comparison, LocalFiles));
@@ -631,8 +631,6 @@ void FAutomationControllerManager::CollectTestResults(TSharedPtr<IAutomationRepo
 			TArray<FString> Keys;
 			Artifact.LocalFiles.GetKeys(Keys);
 
-			bool bOnlyUnapproved = Keys.Num() == 1 && Keys[0] == ComparisonFileTypes::Unapproved;
-
 			ParallelFor(Keys.Num(), [&](int32 Index)
 				{
 					const FString& Key = Keys[Index];
@@ -643,10 +641,11 @@ void FAutomationControllerManager::CollectTestResults(TSharedPtr<IAutomationRepo
 						FScopeLock Lock(&CS);
 						Artifact.Files.Add(Key, MoveTemp(Path));
 					}
-					if (Key == ComparisonFileTypes::Unapproved && Artifact.Type == EAutomationArtifactType::Comparison)
+					if (Artifact.Type == EAutomationArtifactType::Comparison && Key == ComparisonFileTypes::Unapproved)
 					{
-						// Copy screenshot report
-						FScreenshotExportResult ExportResult = ScreenshotManager->ExportScreenshotComparisonResult(Artifact.Name, ArtifactExportPath, bOnlyUnapproved);
+						// Trigger the copy of screenshot comparison artifacts only once, using 'Unapproved' key file as cue (since it is always generated)
+						bool bOnlyGeneratedFiles = Keys.Num() < 3;
+						FScreenshotExportResult ExportResult = ScreenshotManager->ExportScreenshotComparisonResult(Artifact.Name, ArtifactExportPath, bOnlyGeneratedFiles);
 
 						FScopeLock Lock(&CS);
 						if (!JsonTestPassResults.ComparisonExported && ExportResult.Success)
@@ -996,8 +995,9 @@ void FAutomationControllerManager::ProcessResults()
 
 		FAutomatedTestPassResults SerializedPassResults = JsonTestPassResults;
 
-		{ 
-			// Sort result by failure to improve readability
+		if (FParse::Param(FCommandLine::Get(), TEXT("SortTestsByFailure")))
+		{
+			// Sort result by failure to improve readability (disabled by default)
 			SerializedPassResults.Tests.StableSort([](const FAutomatedTestResult& A, const FAutomatedTestResult& B) {
 				if (A.GetErrorTotal() > 0)
 				{
diff --git a/Engine/Source/Developer/ScreenShotComparisonTools/Private/ScreenShotManager.cpp b/Engine/Source/Developer/ScreenShotComparisonTools/Private/ScreenShotManager.cpp
--- a/Engine/Source/Developer/ScreenShotComparisonTools/Private/ScreenShotManager.cpp
+++ b/Engine/Source/Developer/ScreenShotComparisonTools/Private/ScreenShotManager.cpp
@@ -521,7 +521,7 @@ FImageComparisonResult FScreenShotManager::CompareScreenshot(const FString& InUn
 }
 
 
-FScreenshotExportResult FScreenShotManager::ExportScreenshotComparisonResult(FString ScreenshotName, FString RootExportFolder, bool bOnlyIncoming)
+FScreenshotExportResult FScreenShotManager::ExportScreenshotComparisonResult(FString ScreenshotName, FString RootExportFolder, bool bOnlyGeneratedFiles)
 {
 	FPaths::NormalizeDirectoryName(RootExportFolder);
 
@@ -540,9 +540,19 @@ FScreenshotExportResult FScreenShotManager::ExportScreenshotComparisonResult(FSt
 		return Results;
 	}
 
-	FString Pattern = bOnlyIncoming ? TEXT("Incoming.*") : TEXT("*");
+	TArray<FString> Patterns;
+	if (bOnlyGeneratedFiles)
+	{
+		Patterns.Add(TEXT("Incoming.*"));
+		Patterns.Add(TEXT("Delta.png"));
+		Patterns.Add(TEXT("Report.json"));
+	}
+	else
+	{
+		Patterns.Add(TEXT("*"));
+	}
 
-	CopyDirectory(Destination, ScreenshotResultsFolder / ScreenshotName, Pattern);
+	CopyDirectory(Destination, ScreenshotResultsFolder / ScreenshotName, Patterns);
 
 	Results.Success = true;
 	return Results;
@@ -748,13 +758,16 @@ FString FScreenShotManager::GetDefaultExportDirectory() const
 	return FPaths::Combine(FPaths::ProjectSavedDir(),TEXT("Exported/imageCompare"));
 }
 
-void FScreenShotManager::CopyDirectory(const FString& DestDir, const FString& SrcDir, const FString& Pattern)
+void FScreenShotManager::CopyDirectory(const FString& DestDir, const FString& SrcDir, const TArray<FString> Patterns)
 {
 	TArray<FString> FilesToCopy;
 
 	FString AbsoluteSrcDir = FPaths::ConvertRelativePathToFull(SrcDir);
 
-	IFileManager::Get().FindFilesRecursive(FilesToCopy, *AbsoluteSrcDir, *Pattern, /*Files=*/true, /*Directories=*/false);
+	for (auto& Pattern : Patterns)
+	{
+		IFileManager::Get().FindFilesRecursive(FilesToCopy, *AbsoluteSrcDir, *Pattern, /*Files=*/true, /*Directories=*/false, /*bClearFileNames=*/false );
+	}
 
 	ParallelFor(FilesToCopy.Num(), [&](int32 Index)
 		{
diff --git a/Engine/Source/Developer/ScreenShotComparisonTools/Private/ScreenShotManager.h b/Engine/Source/Developer/ScreenShotComparisonTools/Private/ScreenShotManager.h
--- a/Engine/Source/Developer/ScreenShotComparisonTools/Private/ScreenShotManager.h
+++ b/Engine/Source/Developer/ScreenShotComparisonTools/Private/ScreenShotManager.h
@@ -45,7 +45,7 @@ public:
 
 	virtual TFuture<FImageComparisonResult> CompareScreenshotAsync(const FString& IncomingPath, const FAutomationScreenshotMetadata& MetaData, const EScreenShotCompareOptions Options) override;
 
-	virtual FScreenshotExportResult ExportScreenshotComparisonResult(FString ScreenshotName, FString ExportPath = TEXT(""), bool bOnlyIncoming = false) override;
+	virtual FScreenshotExportResult ExportScreenshotComparisonResult(FString ScreenshotName, FString ExportPath = TEXT(""), bool bOnlyGeneratedFiles = false) override;
 
 	TFuture<TSharedPtr<TArray<FComparisonReport>>> OpenComparisonReportsAsync(const FString& ImportPath) override;
 	
@@ -72,7 +72,7 @@ private:
 
 	FImageComparisonResult CompareScreenshot(const FString& IncomingPath, const FAutomationScreenshotMetadata& MetaData, const EScreenShotCompareOptions Options);
 
-	void CopyDirectory(const FString& DestDir, const FString& SrcDir, const FString& Pattern);
+	void CopyDirectory(const FString& DestDir, const FString& SrcDir, const TArray<FString> Patterns);
 
 	void BuildFallbackPlatformsListFromConfig();
 
diff --git a/Engine/Source/Developer/ScreenShotComparisonTools/Public/Interfaces/IScreenShotManager.h b/Engine/Source/Developer/ScreenShotComparisonTools/Public/Interfaces/IScreenShotManager.h
--- a/Engine/Source/Developer/ScreenShotComparisonTools/Public/Interfaces/IScreenShotManager.h
+++ b/Engine/Source/Developer/ScreenShotComparisonTools/Public/Interfaces/IScreenShotManager.h
@@ -68,7 +68,7 @@ public:
 	/**
 	 * Exports target screenshot report to the export location specified
 	 */
-	virtual FScreenshotExportResult ExportScreenshotComparisonResult(FString ScreenshotName, FString ExportPath = TEXT(""), bool bOnlyIncoming = false) = 0;
+	virtual FScreenshotExportResult ExportScreenshotComparisonResult(FString ScreenshotName, FString ExportPath = FString{}, bool bOnlyGeneratedFiles = false) = 0;
 
 	/**
 	 * Imports screenshot comparison data from a given path asynchronously.
