Make sure test state is processed (which include timestamping) even if the json report is not going to be saved right away.

#rnx

diff --git a/Engine/Source/Developer/AutomationController/Private/AutomationControllerManager.cpp b/Engine/Source/Developer/AutomationController/Private/AutomationControllerManager.cpp
--- a/Engine/Source/Developer/AutomationController/Private/AutomationControllerManager.cpp
+++ b/Engine/Source/Developer/AutomationController/Private/AutomationControllerManager.cpp
@@ -833,13 +833,16 @@ void FAutomationControllerManager::ExecuteNextTask( int32 ClusterIndex, OUT bool
 
 						UE_LOG(LogAutomationController, Display, AutomationTestStarting, *NextTest->GetDisplayName(), *NextTest->GetFullTestPath());
 
-						if (JsonTestPassResults.IsRequired && bResumeRunTest)
+						if (JsonTestPassResults.IsRequired)
 						{
 							JsonTestPassResults.UpdateTestResultStatus(NextTest, EAutomationState::InProcess);
 							FAutomatedTestResult& TestResult = JsonTestPassResults.GetTestResult(NextTest);
 							TestResult.DeviceInstance = GameInstances;
-							// Save the whole pass report so that if the next test triggers a critical failure we are not left with no pass report and we can resume.
-							GenerateJsonTestPassSummary(JsonTestPassResults);
+							if (bResumeRunTest)
+							{
+								// Save the whole pass report so that if the next test triggers a critical failure we are not left with no pass report and we can resume.
+								GenerateJsonTestPassSummary(JsonTestPassResults);
+							}
 						}
 
 						NextTest->ResetNetworkCommandResponses();
