Test tag filtering via CLI and CVar using ExecCmd, exposes tag filtering to existing INI groups.
Exposes tag registration from CL 34851931 to users who use the UE Editor console or CLI args. Does not yet plumb tags through non-C++ interfaces, such as Gauntlet or Blueprints. Also does not connect tags to the Automation pane of the Session Frontend UI.

#tests Ran existing tagged tests in editor by executing commands in the console such as: Automation SetTagFilter [TestExampleTag] ... Automation RunTest StartsWith:TestFramework

diff --git a/Engine/Source/Developer/AutomationController/Private/AutomationCommandline.cpp b/Engine/Source/Developer/AutomationController/Private/AutomationCommandline.cpp
--- a/Engine/Source/Developer/AutomationController/Private/AutomationCommandline.cpp
+++ b/Engine/Source/Developer/AutomationController/Private/AutomationCommandline.cpp
@@ -264,6 +264,24 @@ public:
 		}
 	}
 
+
+	void ApplyCVarTagFilter(TSharedPtr <AutomationFilterCollection> AutomationFilters)
+	{
+		if (IConsoleVariable* TagCVar = IConsoleManager::Get().FindConsoleVariable(TEXT("Automation.TestTagGlobalFilter")))
+		{
+			FString TagFilterValue = TagCVar->GetString();
+			if (!TagFilterValue.IsEmpty())
+			{
+				UE_LOG(LogAutomationCommandLine, Display, TEXT("Applying Automation Test tag filter '%s'"), *TagFilterValue);
+				FAutomationGroupFilter* FilterTags = new FAutomationGroupFilter();
+				TArray<FAutomatedTestTagFilter> TagList;
+				TagList.Add(FAutomatedTestTagFilter(TagFilterValue));
+				FilterTags->SetTagFilter(TagList);
+				AutomationFilters->Add(MakeShareable(FilterTags));
+			}
+		}
+	}
+
 	void HandleRefreshTestCallback()
 	{
 		TArray<FString> FilteredTestNames;
@@ -279,6 +297,7 @@ public:
 		// We have found some workers
 		// Create a filter to add to the automation controller, otherwise we don't get any reports
 		TSharedPtr <AutomationFilterCollection> AutomationFilters = MakeShareable(new AutomationFilterCollection());
+		ApplyCVarTagFilter(AutomationFilters);
 		AutomationController->SetFilter(AutomationFilters);
 		AutomationController->SetVisibleTestsEnabled(true);
 		AutomationController->GetEnabledTestNames(FilteredTestNames);
@@ -710,6 +729,15 @@ protected:
 						CVar->Set(true);
 					}
 				}
+				else if (FParse::Command(&TempCmd, TEXT("SetTagFilter")))
+				{
+					if (IConsoleVariable* CVar = IConsoleManager::Get().FindConsoleVariable(TEXT("Automation.TestTagGlobalFilter")))
+					{
+						FString FilterValue(TempCmd);
+						Ar.Logf(TEXT("Automation: Setting test tag global filter to: '%s'"), *FilterValue);
+						CVar->Set(*FilterValue);
+					}
+				}
 				else if (FParse::Command(&TempCmd, TEXT("Help")))
 				{
 					Ar.Logf(TEXT("Supported commands are: "));
@@ -726,6 +754,7 @@ protected:
 					Ar.Logf(TEXT("\tAutomation SoftQuit"));
 					Ar.Logf(TEXT("\tAutomation IgnoreLogEvents"));
 					Ar.Logf(TEXT("\tAutomation EnableStereoTests"));
+					Ar.Logf(TEXT("\tAutomation SetTagFilter"));
 					bHandled = false;
 				}
 				else
diff --git a/Engine/Source/Developer/AutomationController/Private/AutomationReport.cpp b/Engine/Source/Developer/AutomationController/Private/AutomationReport.cpp
--- a/Engine/Source/Developer/AutomationController/Private/AutomationReport.cpp
+++ b/Engine/Source/Developer/AutomationController/Private/AutomationReport.cpp
@@ -42,6 +42,11 @@ FString FAutomationReport::GetTestParameter() const
 	return TestInfo.GetTestParameter();
 }
 
+FString FAutomationReport::GetTags() const
+{
+	return TestInfo.GetTestTags();
+}
+
 FString FAutomationReport::GetAssetPath() const
 {
 	return TestInfo.GetAssetPath();
diff --git a/Engine/Source/Developer/AutomationController/Private/AutomationReport.h b/Engine/Source/Developer/AutomationController/Private/AutomationReport.h
--- a/Engine/Source/Developer/AutomationController/Private/AutomationReport.h
+++ b/Engine/Source/Developer/AutomationController/Private/AutomationReport.h
@@ -22,6 +22,7 @@ public:
 
 	virtual void Empty() override;
 	virtual FString GetTestParameter() const override;
+	virtual FString GetTags() const override;
 	virtual FString GetAssetPath() const override;
 	virtual FString GetOpenCommand() const override;
 	virtual FString GetCommand() const override;
diff --git a/Engine/Source/Developer/AutomationController/Public/AutomationControllerSettings.h b/Engine/Source/Developer/AutomationController/Public/AutomationControllerSettings.h
--- a/Engine/Source/Developer/AutomationController/Public/AutomationControllerSettings.h
+++ b/Engine/Source/Developer/AutomationController/Public/AutomationControllerSettings.h
@@ -72,7 +72,57 @@ public:
 };
 
 /*
-* Describes a filter for a test group with exclude option.
+ * Describes a tag-based filter for tests
+ */
+USTRUCT()
+struct FAutomatedTestTagFilter
+{
+	GENERATED_BODY()
+
+public:
+
+	/**
+	 * @param InContains - String of concatenated tags and boolean operators
+	 * 
+	 * @see FTextFilterExpressionEvaluator
+	 */
+	FAutomatedTestTagFilter(const FString& InContains)
+	{
+		FString FilterString = InContains.TrimStartAndEnd();
+		if (FilterString.IsEmpty())
+		{
+			TagFilter = nullptr;
+		}
+		else
+		{
+			TagFilter = MakeShared<FTextFilterExpressionEvaluator>(ETextFilterExpressionEvaluatorMode::BasicString);
+			TagFilter->SetFilterText(FText::FromString(FilterString));
+		}
+	}
+
+	virtual ~FAutomatedTestTagFilter() {}
+
+	FAutomatedTestTagFilter() : FAutomatedTestTagFilter(TEXT("")) {}
+
+	virtual bool PassesFilter(const TSharedPtr< IAutomationReport >& InReport) const
+	{
+		if (TagFilter)
+		{
+			return TagFilter->TestTextFilter(FBasicStringFilterExpressionContext(InReport->GetTags()));
+		}
+		else // disabled filter, reject nothing
+		{
+			return true;
+		}
+	}
+
+private:
+	TSharedPtr<FTextFilterExpressionEvaluator> TagFilter;
+};
+
+
+/*
+* Describes a filter for a test group with exclude and tag options.
 */
 USTRUCT()
 struct FAutomatedTestFilter : public FAutomatedTestFilterBase
@@ -92,6 +142,10 @@ public:
 	UPROPERTY(Config)
 		TArray<FAutomatedTestFilterBase> Exclude;
 
+	/** List of tag filters specific to this group */
+	UPROPERTY(Config)
+		TArray<FAutomatedTestTagFilter> Tags;
+
 	virtual bool PassesFilter(const TSharedPtr< IAutomationReport >& InReport) const override
 	{
 		bool bMeetsMatch = Super::PassesFilter(InReport);
@@ -108,6 +162,19 @@ public:
 			}
 		}
 
+		// Intersect with tag filters
+		if (bMeetsMatch && !Tags.IsEmpty())
+		{
+			for (const FAutomatedTestTagFilter& Filter : Tags)
+			{
+				if (Filter.PassesFilter(InReport))
+				{
+					return true;
+				}
+			}
+			bMeetsMatch = false; // failed to match any tag filter
+		}
+
 		return bMeetsMatch;
 	}
 };
diff --git a/Engine/Source/Developer/AutomationController/Public/AutomationGroupFilter.h b/Engine/Source/Developer/AutomationController/Public/AutomationGroupFilter.h
--- a/Engine/Source/Developer/AutomationController/Public/AutomationGroupFilter.h
+++ b/Engine/Source/Developer/AutomationController/Public/AutomationGroupFilter.h
@@ -17,12 +17,12 @@ public:
 	{ }
 
 	/** Default constructor with array param. */
-	FAutomationGroupFilter(const TArray<FAutomatedTestFilter> InFilters)
+	FAutomationGroupFilter(const TArray<FAutomatedTestFilter>& InFilters)
 		: Filters(InFilters)
 	{ }
 
 	/** Default constructor, group with single elements. */
-	FAutomationGroupFilter(const FAutomatedTestFilter InFilter)
+	FAutomationGroupFilter(const FAutomatedTestFilter& InFilter)
 	{
 		Filters.Add(InFilter);
 	}
@@ -31,14 +31,26 @@ public:
 	/**
 	 * Set the list of strings the group filter checks for substrings in test display name.
 	 *
-	 * @param InContainsArray An array of strings to filter against test display names.
-	 * @see SetMatchFromStartArray, SetMatchFromEndArray
+	 * @param InFilters An array of strings to filter against test display names.
+	 * @see FAutomatedTestFilterBase::MatchFromStart, FAutomatedTestFilterBase::MatchFromEnd
 	 */
-	void SetFilters(const TArray<FAutomatedTestFilter> InFilters)
+	void SetFilters(const TArray<FAutomatedTestFilter>& InFilters)
 	{
 		Filters = InFilters;
 		ChangedEvent.Broadcast();
 	}
+
+	/**
+	 * Set the list of search syntax strings evaluated against test tags.
+	 *
+	 * @param InFilters An array of strings to filter against test tags.
+	 * @see FTextFilterExpressionEvaluator
+	 */
+	void SetTagFilter(const TArray < FAutomatedTestTagFilter>& InFilters)
+	{
+		TagFilters = InFilters;
+		ChangedEvent.Broadcast();
+	}
 public:
 
 	// IFilter interface
@@ -46,17 +58,31 @@ public:
 	DECLARE_DERIVED_EVENT(FAutomationGroupFilter, IFilter< const TSharedPtr< class IAutomationReport >& >::FChangedEvent, FChangedEvent);
 	virtual FChangedEvent& OnChanged() override { return ChangedEvent; }
 
-	virtual bool PassesFilter( const TSharedPtr< IAutomationReport >& InReport ) const override
+	virtual bool PassesFilter(const TSharedPtr< IAutomationReport >& InReport) const override
 	{
-		for (const FAutomatedTestFilter& Filter: Filters)
+		// empty filters pass against all values
+		bool NameFilterPassing = Filters.IsEmpty();
+		bool TagFilterPassing = TagFilters.IsEmpty();
+
+		for (const FAutomatedTestFilter& NameFilter: Filters)
 		{
-			if (Filter.PassesFilter(InReport))
+			if (NameFilter.PassesFilter(InReport))
 			{
-				return true;
+				NameFilterPassing = true;
+				break;
 			}
 		}
 
-		return Filters.Num() == 0;
+		for (const FAutomatedTestTagFilter& TagFilter : TagFilters)
+		{
+			if (TagFilter.PassesFilter(InReport))
+			{
+				TagFilterPassing = true;
+				break;
+			}
+		}
+
+		return NameFilterPassing && TagFilterPassing;
 	}
 
 private:
@@ -64,7 +90,10 @@ private:
 	/**	The event that broadcasts whenever a change occurs to the filter. */
 	FChangedEvent ChangedEvent;
 
-	/** The array of FAutomatedTestFilter to filter against. At least one from the list must be matched. */
+	/** The array of FAutomatedTestFilter to filter against test names. At least one from the list must be matched. */
 	TArray<FAutomatedTestFilter> Filters;
 
+	/** The array of FAutomatedTestTagFilter to filter against test tags. At least one from the list must be matched. */
+	TArray<FAutomatedTestTagFilter> TagFilters;
+
 };
diff --git a/Engine/Source/Developer/AutomationController/Public/IAutomationReport.h b/Engine/Source/Developer/AutomationController/Public/IAutomationReport.h
--- a/Engine/Source/Developer/AutomationController/Public/IAutomationReport.h
+++ b/Engine/Source/Developer/AutomationController/Public/IAutomationReport.h
@@ -234,6 +234,13 @@ public:
 	 */
 	virtual FString GetTestParameter() const = 0;
 
+	/**
+	 * Get the tags associated with this test.
+	 *
+	 * @return the tags concatenated into a single string.
+	 */
+	virtual FString GetTags() const = 0;
+
 	/**
 	 * Gets the asset path associated with a test, it may not have one.
 	 * 
diff --git a/Engine/Source/Runtime/AutomationMessages/Public/AutomationWorkerMessages.h b/Engine/Source/Runtime/AutomationMessages/Public/AutomationWorkerMessages.h
--- a/Engine/Source/Runtime/AutomationMessages/Public/AutomationWorkerMessages.h
+++ b/Engine/Source/Runtime/AutomationMessages/Public/AutomationWorkerMessages.h
@@ -240,6 +240,9 @@ struct FAutomationWorkerSingleTestReply : public FAutomationWorkerMessageBase
 	UPROPERTY(EditAnywhere, Category="Message")
 	uint32 NumParticipantsRequired = 0;
 
+	UPROPERTY(EditAnywhere, Category="Message")
+	FString TestTags;
+
 	/** Default constructor. */
 	FAutomationWorkerSingleTestReply() = default;
 
@@ -256,6 +259,7 @@ struct FAutomationWorkerSingleTestReply : public FAutomationWorkerMessageBase
 		OpenCommand = InTestInfo.GetOpenCommand();
 		TestFlags = (uint32)InTestInfo.GetTestFlags();
 		NumParticipantsRequired = InTestInfo.GetNumParticipantsRequired();
+		TestTags = InTestInfo.GetTestTags();
 	}
 
 	FAutomationTestInfo GetTestInfo() const
@@ -270,7 +274,8 @@ struct FAutomationWorkerSingleTestReply : public FAutomationWorkerMessageBase
 			SourceFile,
 			SourceFileLine,
 			AssetPath,
-			OpenCommand);
+			OpenCommand,
+			TestTags);
 	}
 };
 
diff --git a/Engine/Source/Runtime/Core/Private/Misc/AutomationTest.cpp b/Engine/Source/Runtime/Core/Private/Misc/AutomationTest.cpp
--- a/Engine/Source/Runtime/Core/Private/Misc/AutomationTest.cpp
+++ b/Engine/Source/Runtime/Core/Private/Misc/AutomationTest.cpp
@@ -64,6 +64,12 @@ namespace AutomationTest
 		bLightweightStereoTestVariants,
 		TEXT("Whether to skip variants when the baseline test fails, and skip saving screenshots for successful variants"));
 
+	FString TestTagGlobalFilter = "";
+	static FAutoConsoleVariableRef CVarAutomationTestTagGlobalFilter(
+		TEXT("Automation.TestTagGlobalFilter"),
+		TestTagGlobalFilter,
+		TEXT("Only include tests marked with Tags matching this filter string, using the Advanced Search Syntax"));
+
 	// The method prepares the filename and LineNumber to be placed in the form that could be extracted by SAutomationWindow widget if it is additionally eclosed into []
 	// The result format is filename(line)
 	static FString CreateFileLineDescription(const FString& Filename, const int32 LineNumber)
@@ -366,9 +372,9 @@ bool FAutomationTestFramework::UnregisterAutomationTestTags(const FString& InTes
 	return bRegistered;
 }
 
-bool FAutomationTestFramework::RegisterComplexAutomationTestTags(const FAutomationTestBase* InTest, const FString& InTestNameToRegister, const FString& InTestTagsToRegister)
+bool FAutomationTestFramework::RegisterComplexAutomationTestTags(const FAutomationTestBase* InTest, const FString& InBeautifiedTestName, const FString& InTestTagsToRegister)
 {
-	FString FullTestName = InTest->GetBeautifiedTestName().AppendChar('.').Append(InTestNameToRegister);
+	FString FullTestName = InTest->GetBeautifiedTestName().AppendChar('.').Append(InBeautifiedTestName);
 	return RegisterAutomationTestTags(FullTestName, InTestTagsToRegister);
 }
 
@@ -1704,7 +1710,7 @@ void FAutomationTestBase::GenerateTestNames(TArray<FAutomationTestInfo>& TestInf
 			GetTestSourceFileLine(CompleteTestName),
 			GetTestAssetPath(ParameterNames[ParameterIndex]),
 			GetTestOpenCommand(ParameterNames[ParameterIndex]),
-			Framework.GetTagsForAutomationTest(TestName)
+			Framework.GetTagsForAutomationTest(CompleteBeautifiedNames)
 		);
 		
 		TestInfo.Add( NewTestInfo );
