Fix compilation error 'UE_FMT_STR_Checker' (aka 'int') is not a class, namespace, or enumeration

#rnx

diff --git a/Engine/Source/Developer/AutomationController/Private/AutomationCommandline.cpp b/Engine/Source/Developer/AutomationController/Private/AutomationCommandline.cpp
--- a/Engine/Source/Developer/AutomationController/Private/AutomationCommandline.cpp
+++ b/Engine/Source/Developer/AutomationController/Private/AutomationCommandline.cpp
@@ -468,7 +468,7 @@ public:
 							GIsCriticalError = true;
 						}
 					}
-					UE_LOG(LogAutomationCommandLine, Log, TEXT("Shutting down. GIsCriticalError=%d"), GIsCriticalError);
+					UE_LOG(LogAutomationCommandLine, Log, TEXT("Shutting down. GIsCriticalError=%d"), GIsCriticalError ? 1 : 0);
 					// some tools parse this.
 					UE_LOG(LogAutomationCommandLine, Display, TEXT("**** TEST COMPLETE. EXIT CODE: %d ****"), GIsCriticalError ? -1 : 0);
 					FPlatformMisc::RequestExitWithStatus(AutomationCommand == EAutomationCommand::SoftQuit ? false : true, GIsCriticalError ? -1 : 0);
