PR #13158: Allow Config of Automation Controller Settings

diff --git a/Engine/Source/Developer/AutomationController/Public/AutomationControllerSettings.h b/Engine/Source/Developer/AutomationController/Public/AutomationControllerSettings.h
--- a/Engine/Source/Developer/AutomationController/Public/AutomationControllerSettings.h
+++ b/Engine/Source/Developer/AutomationController/Public/AutomationControllerSettings.h
@@ -240,11 +240,11 @@ public:
 	bool bAutoExpandSingleItemSubgroups;
 
 	/** Whether to Sort tests by failure type in json/html report */
-	UPROPERTY()
+	UPROPERTY(Config)
 	bool bSortTestsByFailure;
 
 	/** Whether to prune log events from test report on success */
-	UPROPERTY()
+	UPROPERTY(Config)
 	bool bPruneLogsOnSuccess;
 
 private:
