Avoid loading config for all platforms when loading FAutomationControllerModule
- UAutomationTestExcludelist automatically loads config the first time it is accessed. The automation module has a tree structure of FAutomationReports and creates an empty root node during init. This empty root node was the only thing calling into UAutomationTestExcludelist during loading. Exclude data isn't actually needed until the automation window is opened. This change ends up saving 340ms during editor startup.
- A better approach would be to have UAutomationTestExcludelist lazy load config for individual platforms as necessary. Unfortunately, the automation window implicitly relies on FAutomationReportManager automatically loading config for all platforms. It would need to be updated to explicitly request config for all desired platforms. As it stands, this unnecessary loading during startup is likely to creep back in silently in the future.

diff --git a/Engine/Source/Developer/AutomationController/Private/AutomationReport.cpp b/Engine/Source/Developer/AutomationController/Private/AutomationReport.cpp
--- a/Engine/Source/Developer/AutomationController/Private/AutomationReport.cpp
+++ b/Engine/Source/Developer/AutomationController/Private/AutomationReport.cpp
@@ -3,9 +3,6 @@
 #include "AutomationReport.h"
 #include "AutomationTestExcludelist.h"
 #include "Misc/FilterCollection.h"
-#include "HAL/FileManager.h"
-#include "Misc/FileHelper.h"
-#include "Misc/Paths.h"
 
 FAutomationReport::FAutomationReport(FAutomationTestInfo& InTestInfo, bool InIsParent)
 	: bEnabled( false )
@@ -20,12 +17,19 @@ FAutomationReport::FAutomationReport(FAutomationTestInfo& InTestInfo, bool InIsP
 	{
 		bEnabled = true;
 	}
-	// Get exclude test info from Config/DefaultEngine.ini
-	if (auto Entry = UAutomationTestExcludelist::Get()->GetExcludeTestEntry(TestInfo.GetFullTestPath()))
+
+	// Get exclude test info from config for every supported platform. We want to avoid querying when TestInfo is empty
+	// (e.g. this is the empty root node in FAutomationControllerManager and we don't actually need to load config yet).
+	// This saves ~340ms when opening the editor.
+	const FString& TestPath = TestInfo.GetFullTestPath();
+	if (!TestPath.IsEmpty())
 	{
-		ExcludeTestInfo = *Entry;
-		ExcludeTestInfo.SetPropagation(TestInfo.GetFullTestPath());
-		bNeedToSkip = true;
+		if (auto Entry = UAutomationTestExcludelist::Get()->GetExcludeTestEntry(TestPath))
+		{
+			ExcludeTestInfo = *Entry;
+			ExcludeTestInfo.SetPropagation(TestPath);
+			bNeedToSkip = true;
+		}
 	}
 }
 
