AutomationController: Prevent referencing editor module for client targets

#rnx

#Synced-CL 50547795

diff --git a/Engine/Source/Developer/AutomationController/AutomationController.Build.cs b/Engine/Source/Developer/AutomationController/AutomationController.Build.cs
--- a/Engine/Source/Developer/AutomationController/AutomationController.Build.cs
+++ b/Engine/Source/Developer/AutomationController/AutomationController.Build.cs
@@ -20,8 +20,6 @@ namespace UnrealBuildTool.Rules
 				{
                     "AssetRegistry",
 					"AutomationMessages",
-					"UnrealEdMessages",
-                    "MessageLog",
                     "Json",
                     "JsonUtilities",
 					"ScreenShotComparisonTools",
@@ -35,7 +33,9 @@ namespace UnrealBuildTool.Rules
                 PrivateDependencyModuleNames.AddRange(
                     new string[] {
 						"EditorFramework",
-                        "UnrealEd",
+						"MessageLog",
+						"UnrealEdMessages",
+						"UnrealEd",
                         "Engine", // Needed for UWorld/GWorld to find current level
 				    }
                 );
diff --git a/Engine/Source/Developer/AutomationController/Private/AutomationControllerManager.cpp b/Engine/Source/Developer/AutomationController/Private/AutomationControllerManager.cpp
--- a/Engine/Source/Developer/AutomationController/Private/AutomationControllerManager.cpp
+++ b/Engine/Source/Developer/AutomationController/Private/AutomationControllerManager.cpp
@@ -14,7 +14,6 @@
 #include "MessageEndpoint.h"
 #include "MessageEndpointBuilder.h"
 #include "Modules/ModuleManager.h"
-#include "AssetEditorMessages.h"
 #include "ImageComparer.h"
 #include "AutomationControllerSettings.h"
 #include "Interfaces/IScreenShotToolsModule.h"
@@ -28,6 +27,7 @@
 
 #if WITH_EDITOR
 #include "UnrealEdGlobals.h"
+#include "AssetEditorMessages.h"
 #include "Editor/UnrealEdEngine.h"
 #include "Logging/MessageLog.h"
 #include "Tests/AutomationCommon.h"
@@ -446,7 +446,9 @@ void FAutomationControllerManager::Init()
 
 void FAutomationControllerManager::RequestLoadAsset(const FString& InAssetName)
 {
+#if WITH_EDITOR
 	MessageEndpoint->Publish(FMessageEndpoint::MakeMessage<FAssetEditorRequestOpenAsset>(InAssetName), EMessageScope::Process);
+#endif
 }
 
 void FAutomationControllerManager::Tick()
