Apply tag filter: save tags then use with existing filters to restrict test list correctly

#rnx

#Synced-CL 0

diff --git a/Engine/Source/Developer/AutomationController/Private/AutomationCommandline.cpp b/Engine/Source/Developer/AutomationController/Private/AutomationCommandline.cpp
--- a/Engine/Source/Developer/AutomationController/Private/AutomationCommandline.cpp
+++ b/Engine/Source/Developer/AutomationController/Private/AutomationCommandline.cpp
@@ -160,7 +160,14 @@ public:
 					FilterName += TEXT(".");
 				}
 
-				FiltersList.Add(FAutomatedTestFilter(FilterName, true, false));
+				FAutomatedTestFilter StartsWithFilter(FilterName, true, false);
+
+				if (!TagFilter.IsEmpty())
+				{
+					StartsWithFilter.Tags = GetTestTagFilter(TagFilter);
+				}
+
+				FiltersList.Add(StartsWithFilter);
 			}
 			else if (ArgumentName.StartsWith(GroupPrefix))
 			{
@@ -178,6 +185,15 @@ public:
 						// if found add all this groups filters to our current list
 						if (GroupEntry->Filters.Num() > 0)
 						{
+							if (!TagFilter.IsEmpty())
+							{
+								auto TagsFilter = GetTestTagFilter(TagFilter);
+								for (int32 j = 0; j < GroupEntry->Filters.Num(); ++j)
+								{
+									FAutomatedTestFilter* Filter = &(GroupEntry->Filters[j]);
+									Filter->Tags = TagsFilter;
+								}
+							}
 							FiltersList.Append(GroupEntry->Filters);
 						}
 						else
@@ -191,7 +207,7 @@ public:
 				{
 					UE_LOG(LogAutomationCommandLine, Error, TEXT("No matching group named %s"), *GroupName);
 				}
-			}			
+			}
 			else
 			{
 				bool bMatchFromStart = false;
@@ -276,14 +292,19 @@ public:
 
 				UE_LOG(LogAutomationCommandLine, Display, TEXT("Applying Automation Test tag filter '%s'"), *TagFilterValue);
 				FAutomationGroupFilter* FilterTags = new FAutomationGroupFilter();
-				TArray<FAutomatedTestTagFilter> TagList;
-				TagList.Add(FAutomatedTestTagFilter(TagFilterValue));
-				FilterTags->SetTagFilter(TagList);
+				FilterTags->SetTagFilter(GetTestTagFilter(TagFilterValue));
 				AutomationFilters->Add(MakeShareable(FilterTags));
 			}
 		}
 	}
 
+	TArray<FAutomatedTestTagFilter> GetTestTagFilter(FString InTagFilter)
+	{
+		TArray<FAutomatedTestTagFilter> TagList;
+		TagList.Add(FAutomatedTestTagFilter(InTagFilter));
+		return TagList;
+	}
+
 	void HandleRefreshTestCallback()
 	{
 		TArray<FString> FilteredTestNames;
@@ -324,7 +345,14 @@ public:
 			
 			if (FilteredTestNames.Num() == 0)
 			{
-				LogCommandLineError(FString::Printf(TEXT("No automation tests matched '%s'"), *StringCommand));	
+				if (bAllowZeroTestResults)
+				{
+					UE_LOG(LogAutomationCommandLine, Warning, TEXT("No automation tests matched '%s'"), *StringCommand);
+				}
+				else
+				{
+					LogCommandLineError(FString::Printf(TEXT("No automation tests matched '%s'"), *StringCommand));
+				}
 			}
 			else
 			{
@@ -740,6 +768,15 @@ protected:
 						CVar->Set(*FilterValue);
 					}
 				}
+				else if (FParse::Command(&TempCmd, TEXT("ApplyTagFilter")))
+				{
+					TagFilter = FString(TempCmd).TrimStartAndEnd();
+					Ar.Logf(TEXT("Automation: Storing tag filter to be applied: '%s'"), *TagFilter);
+				}
+				else if (FParse::Command(&TempCmd, TEXT("AllowZeroTestResults")))
+				{
+					bAllowZeroTestResults = true;
+				}
 				else if (FParse::Command(&TempCmd, TEXT("Help")))
 				{
 					Ar.Logf(TEXT("Supported commands are: "));
@@ -757,6 +794,8 @@ protected:
 					Ar.Logf(TEXT("\tAutomation IgnoreLogEvents"));
 					Ar.Logf(TEXT("\tAutomation EnableStereoTests"));
 					Ar.Logf(TEXT("\tAutomation SetTagFilter"));
+					Ar.Logf(TEXT("\tAutomation ApplyTagFilter"));
+					Ar.Logf(TEXT("\tAutomation AllowZeroTestResults"));
 					bHandled = false;
 				}
 				else
@@ -821,6 +860,12 @@ private:
 
 	// Any that we encountered during processing. Used in 'Quit' to determine error code
 	TArray<FString> Errors;
+
+	// Tag filter expression to be set to each group
+	FString TagFilter;
+
+	// If true, do not return non-zero exit code when no tests were run
+	bool bAllowZeroTestResults = false;
 };
 
 const float FAutomationExecCmd::DefaultDelayTimer = 5.0f;
diff --git a/Engine/Source/Programs/AutomationTool/Gauntlet/Framework/Horde/Gauntlet.Horde.cs b/Engine/Source/Programs/AutomationTool/Gauntlet/Framework/Horde/Gauntlet.Horde.cs
--- a/Engine/Source/Programs/AutomationTool/Gauntlet/Framework/Horde/Gauntlet.Horde.cs
+++ b/Engine/Source/Programs/AutomationTool/Gauntlet/Framework/Horde/Gauntlet.Horde.cs
@@ -114,6 +114,9 @@ namespace Gauntlet
 			public static HordeStepMessage TestSucceeded = new HordeStepMessage("TestSucceeded", JobStepMessageLevels.Info);
 			public static HordeStepMessage TestDataGenerated = new HordeStepMessage("TestDataGenerated", JobStepMessageLevels.Info);
 
+			// Failed states
+			public static HordeStepMessage ZeroTestsRun = new HordeStepMessage("ZeroTestsRun", JobStepMessageLevels.Error);
+
 		}
 
 		private static string? curStepMessage { get; set; } = null;
diff --git a/Engine/Source/Programs/AutomationTool/Gauntlet/Unreal/Automation/UE.Automation.cs b/Engine/Source/Programs/AutomationTool/Gauntlet/Unreal/Automation/UE.Automation.cs
--- a/Engine/Source/Programs/AutomationTool/Gauntlet/Unreal/Automation/UE.Automation.cs
+++ b/Engine/Source/Programs/AutomationTool/Gauntlet/Unreal/Automation/UE.Automation.cs
@@ -223,6 +223,17 @@ namespace UE
 
 			string AutomationTestArgument = string.Format("RunTest {0};", RunTest);
 
+			if (!string.IsNullOrEmpty(TagFilter))
+			{
+				AutomationTestArgument = $"ApplyTagFilter {TagFilter};{AutomationTestArgument}";
+			}
+
+			bool bAllowZeroTestResults = Globals.Params.ParseParam("AllowZeroTestResults");
+			if (bAllowZeroTestResults)
+			{
+				AutomationTestArgument = $"AllowZeroTestResults;{AutomationTestArgument}";
+			}
+
 			// if this is not attended then we'll quit the editor after the tests and disable any user interactions
 			if (this.Attended == false)
 			{
@@ -369,7 +380,6 @@ namespace UE
 					.Select(LogName => LogName + " Verbose"));
 				AppConfig.CommandLine += string.Format(" -LogCmds=\"{0}\"", LogCmdsString);
 			}
-
 		}
 	}
 
@@ -1015,9 +1025,19 @@ namespace UE
 					// If no tests were run then that's a failure (possibly a bad RunTest argument?)
 					if (TestResults == null || TestResults.Tests.Count == 0)
 					{
-						ExitReason = "No tests were executed!";
-						ExitCode = -1;
-						return UnrealProcessResult.InitializationFailure;
+						bool bAllowZeroTestResults = Globals.Params.ParseParam("AllowZeroTestResults");
+						if (bAllowZeroTestResults)
+						{
+							ExitReason = "No tests were executed, but AllowZeroTestResults is set, so treating this as a success.";
+							ExitCode = 0;
+							return UnrealProcessResult.ExitOk;
+						}
+						else
+						{
+							ExitReason = "No tests were executed!";
+							ExitCode = -1;
+							return UnrealProcessResult.InitializationFailure;
+						}
 					}
 				}
 			}
@@ -1114,6 +1134,8 @@ namespace UE
 		{
 			base.LogTestSummaryHeader();
 
+			bool bAllowZeroTestResults = Globals.Params.ParseParam("AllowZeroTestResults");
+
 			// Everything we need is in the editor artifacts
 			var EditorRole = RoleResults.Where(R => R.Artifacts.SessionRole.RoleType.IsEditor()).FirstOrDefault();
 
@@ -1140,8 +1162,16 @@ namespace UE
 				{
 					if (AllTests.Count == 0)
 					{
-						Log.Error(KnownLogEvents.Gauntlet_UnrealEngineTestEvent, " * No tests were executed.");
-						Log.Info("");
+						if (bAllowZeroTestResults)
+						{
+							Log.Warning(KnownLogEvents.Gauntlet_UnrealEngineTestEvent, " * No tests were executed, check the test filter.");
+							Log.Info("");
+						}
+						else
+						{
+							Log.Error(KnownLogEvents.Gauntlet_UnrealEngineTestEvent, " * No tests were executed.");
+							Log.Info("");
+						}
 					}
 					else if (FailedTests.Count > 0 || IncompleteTests.Count > 0)
 					{
@@ -1186,7 +1216,15 @@ namespace UE
 				{
 					if (AllTests.Count == 0)
 					{
-						Log.Error(KnownLogEvents.Gauntlet_UnrealEngineTestEvent, " * No tests were executed.");
+						if (bAllowZeroTestResults)
+						{
+							Log.Warning(KnownLogEvents.Gauntlet_UnrealEngineTestEvent, " * No tests were executed, check test filter.");
+							Gauntlet.Horde.UpdateStepMessage(Gauntlet.Horde.GauntletHordeStepMessage.ZeroTestsRun);
+						}
+						else
+						{
+							Log.Error(KnownLogEvents.Gauntlet_UnrealEngineTestEvent, " * No tests were executed.");
+						}
 
 						IEnumerable<UnrealTestEvent> Errors = TestNodeEvents.Where(E => E.IsError).Distinct();
 						IEnumerable<UnrealTestEvent> Warnings = TestNodeEvents.Where(E => E.IsWarning).Distinct();
diff --git a/Engine/Source/Programs/AutomationTool/Gauntlet/Unreal/Base/Gauntlet.UnrealTestNode.cs b/Engine/Source/Programs/AutomationTool/Gauntlet/Unreal/Base/Gauntlet.UnrealTestNode.cs
--- a/Engine/Source/Programs/AutomationTool/Gauntlet/Unreal/Base/Gauntlet.UnrealTestNode.cs
+++ b/Engine/Source/Programs/AutomationTool/Gauntlet/Unreal/Base/Gauntlet.UnrealTestNode.cs
@@ -1789,8 +1789,17 @@ namespace Gauntlet
 						}
 						else
 						{
+							bool bAllowZeroTestResults = Globals.Params.ParseParam("AllowZeroTestResults");
+
 							var Stream = InitPhase.GetStream();
-							Stream.AddError("No tests were executed!");
+							if (bAllowZeroTestResults)
+							{
+								Stream.AddWarning("No tests were executed, check test filter!");
+							}
+							else
+							{
+								Stream.AddError("No tests were executed!");
+							}
 							if (RoleResults != null)
 							{
 								// Add any relevant events collected through log parsing
diff --git a/Engine/Source/Programs/AutomationTool/Gauntlet/Unreal/Engine/UnrealTests.cs b/Engine/Source/Programs/AutomationTool/Gauntlet/Unreal/Engine/UnrealTests.cs
--- a/Engine/Source/Programs/AutomationTool/Gauntlet/Unreal/Engine/UnrealTests.cs
+++ b/Engine/Source/Programs/AutomationTool/Gauntlet/Unreal/Engine/UnrealTests.cs
@@ -265,7 +265,16 @@ namespace Gauntlet
 
 					if (AllTests.Count() == 0)
 					{
-						Log.Error(KnownLogEvents.Gauntlet_UnrealEngineTestEvent, " * No tests were executed!");
+						bool bAllowZeroTestResults = Globals.Params.ParseParam("AllowZeroTestResults");
+
+						if (bAllowZeroTestResults)
+						{
+							Log.Warning(KnownLogEvents.Gauntlet_UnrealEngineTestEvent, " * No tests were executed, check test filter!");
+						}
+						else
+						{
+							Log.Error(KnownLogEvents.Gauntlet_UnrealEngineTestEvent, " * No tests were executed!");
+						}
 					}
 					else
 					{
diff --git a/Engine/Source/Programs/AutomationTool/Gauntlet/Unreal/RunUnreal.cs b/Engine/Source/Programs/AutomationTool/Gauntlet/Unreal/RunUnreal.cs
--- a/Engine/Source/Programs/AutomationTool/Gauntlet/Unreal/RunUnreal.cs
+++ b/Engine/Source/Programs/AutomationTool/Gauntlet/Unreal/RunUnreal.cs
@@ -258,11 +258,6 @@ namespace Gauntlet
 						Role.ExtraArgs += string.Format(" -ExecCmds=\"{0}\"", RoleExecCmds);
 					}
 
-					if (!string.IsNullOrEmpty(TagFilter))
-					{
-						Role.ExtraArgs += string.Format(" -ExecCmds=\"Automation SetTagFilter {0}\"", TagFilter);
-					}
-
 					bool UsesEditor = EditorForAllRoles || Globals.Params.ParseParam("Editor" + Type.ToString());
 
 					if (UsesEditor)
