Allow for emitting testing formats the same way the Automation Controller does

#Synced-CL 52522074

diff --git a/Engine/Source/Developer/AutomationController/Private/AutomationControllerManager.cpp b/Engine/Source/Developer/AutomationController/Private/AutomationControllerManager.cpp
--- a/Engine/Source/Developer/AutomationController/Private/AutomationControllerManager.cpp
+++ b/Engine/Source/Developer/AutomationController/Private/AutomationControllerManager.cpp
@@ -1,6 +1,8 @@
 // Copyright Epic Games, Inc. All Rights Reserved.
 
 #include "AutomationControllerManager.h"
+#include "IAutomationControllerModule.h"
+#include "Misc/AutomationEvent.h"
 #include "CoreMinimal.h"
 #include "HAL/FileManager.h"
 #include "Misc/CommandLine.h"
@@ -45,6 +47,42 @@
 
 DEFINE_LOG_CATEGORY_STATIC(LogAutomationController, Log, All)
 
+void FAutomationResultEmitter::EmitTestStarted(const FString& Name, const FString& Path)
+{
+	UE_LOG(LogAutomationController, Display, AutomationTestStarting, *Name, *Path);
+}
+
+void FAutomationResultEmitter::EmitTestCompleted(const FString& State, const FString& Name, const FString& Path)
+{
+	UE_LOG(LogAutomationController, Display, AutomationStateFormat, *State, *Name, *Path);
+}
+
+void FAutomationResultEmitter::EmitBeginEvents(const FString& Path)
+{
+	UE_LOG(LogAutomationController, Log, BeginEventsFormat, *Path);
+}
+
+void FAutomationResultEmitter::EmitEndEvents(const FString& Path)
+{
+	UE_LOG(LogAutomationController, Log, EndEventsFormat, *Path);
+}
+
+void FAutomationResultEmitter::EmitEntry(const FAutomationExecutionEntry& Entry)
+{
+	switch (Entry.Event.Type)
+	{
+	case EAutomationEventType::Error:
+		UE_LOG(LogAutomationController, Error, TEXT("%s"), *Entry.ToString());
+		break;
+	case EAutomationEventType::Warning:
+		UE_LOG(LogAutomationController, Warning, TEXT("%s"), *Entry.ToString());
+		break;
+	default:
+		UE_LOG(LogAutomationController, Log, TEXT("%s"), *Entry.ToString());
+		break;
+	}
+}
+
 #define LOCTEXT_NAMESPACE "AutomationTesting"
 
 void FAutomatedTestPassResults::AddTestResult(const IAutomationReportPtr& TestReport)
diff --git a/Engine/Source/Developer/AutomationController/Public/IAutomationControllerModule.h b/Engine/Source/Developer/AutomationController/Public/IAutomationControllerModule.h
--- a/Engine/Source/Developer/AutomationController/Public/IAutomationControllerModule.h
+++ b/Engine/Source/Developer/AutomationController/Public/IAutomationControllerModule.h
@@ -7,6 +7,26 @@
 #include "Modules/ModuleManager.h"
 #include "IAutomationControllerManager.h"
 
+struct FAutomationExecutionEntry;
+
+/**
+ * Emit test results in the format Gauntlet's AutomationLogParser expects.
+ * Shares the file-scoped format #defines with ReportAutomationResult in
+ * AutomationControllerManager.cpp. Both call sites use the same defines
+ * but are not yet consolidated into a single code path.
+ *
+ * TODO: Refactor ReportAutomationResult to call these functions, making
+ * them the single entry point for Gauntlet protocol emission.
+ */
+struct FAutomationResultEmitter
+{
+	AUTOMATIONCONTROLLER_API static void EmitTestStarted(const FString& Name, const FString& Path);
+	AUTOMATIONCONTROLLER_API static void EmitTestCompleted(const FString& State, const FString& Name, const FString& Path);
+	AUTOMATIONCONTROLLER_API static void EmitBeginEvents(const FString& Path);
+	AUTOMATIONCONTROLLER_API static void EmitEndEvents(const FString& Path);
+	AUTOMATIONCONTROLLER_API static void EmitEntry(const FAutomationExecutionEntry& Entry);
+};
+
 /**
  * Interface for AutomationController modules.
  */
